Merging CL 3074131 from 4.13 branch to fix DDC warnings.

diff --git a/Templates/TP_VirtualRealityBP/Config/DefaultEngine.ini b/Templates/TP_VirtualRealityBP/Config/DefaultEngine.ini
--- a/Templates/TP_VirtualRealityBP/Config/DefaultEngine.ini
+++ b/Templates/TP_VirtualRealityBP/Config/DefaultEngine.ini
@@ -11,8 +11,6 @@ ServerDefaultMap=/Engine/Maps/Entry
 GlobalDefaultGameMode=None
 GlobalDefaultServerGameMode=None
 
-
-
 [/Script/Engine.Engine]
 +ActiveGameNameRedirects=(OldGameName="TP_FirstPersonBP",NewGameName="/Script/VRTemplate")
 +ActiveGameNameRedirects=(OldGameName="/Script/TP_FirstPersonBP",NewGameName="/Script/VRTemplate")
@@ -30,23 +28,21 @@ vr.InstancedStereo=True
 r.SeparateTranslucency=0
 r.HZBOcclusion=0
 r.FinishCurrentFrame=1
-r.MotionBlurQuality=0
-r.PostProcessAAQuality=3
-r.BloomQuality=1
 r.EyeAdaptionQuality=0
-r.AmbientOcclusionLevels=0
-r.SSR.Quality=0
-r.DepthOfFieldQuality=0
-r.SceneColorFormat=2
-r.TranslucencyVolumeBlur=0
-r.TranslucencyLightingVolumeDim=4
-r.MaxAnisotropy=8
-r.LensFlareQuality=0
-r.SceneColorFringeQuality=0
-r.FastBlurThreshold=0
 r.SSR.MaxRoughness=0
 r.rhicmdbypass=0
 r.TiledReflectionEnvironmentMinimumCount=10
+r.MobileHDR=False
+
+[/Script/Engine.NavigationSystem]
+bAutoCreateNavigationData=True
+bAllowClientSideNavigation=False
+bInitialBuildingLocked=False
+bSkipAgentHeightCheckWhenPickingNavData=False
+DataGatheringMode=Instant
+bGenerateNavigationOnlyAroundNavigationInvokers=False
+ActiveTilesUpdateInterval=1.000000
+DirtyAreasUpdateFreq=60.000000
 
 [/Script/Engine.UserInterfaceSettings]
 RenderFocusRule=NavigationOnly
@@ -62,12 +58,3 @@ UIScaleRule=ShortestSide
 CustomScalingRuleClass=None
 UIScaleCurve=(EditorCurveData=(PreInfinityExtrap=RCCE_Constant,PostInfinityExtrap=RCCE_Constant,Keys=((Time=480.000000,Value=0.444000),(Time=720.000000,Value=0.666000),(Time=1080.000000,Value=1.000000),(Time=8640.000000,Value=8.000000)),DefaultValue=340282346638528859811704183484516925440.000000),ExternalCurve=None)
 
-[/Script/Engine.NavigationSystem]
-bAutoCreateNavigationData=True
-bAllowClientSideNavigation=False
-bInitialBuildingLocked=False
-bSkipAgentHeightCheckWhenPickingNavData=False
-DataGatheringMode=Instant
-bGenerateNavigationOnlyAroundNavigationInvokers=False
-ActiveTilesUpdateInterval=1.000000
-DirtyAreasUpdateFreq=60.000000
diff --git a/Templates/TP_VirtualRealityBP/Config/DefaultScalability.ini b/Templates/TP_VirtualRealityBP/Config/DefaultScalability.ini
new file mode 100644
--- /dev/null
+++ b/Templates/TP_VirtualRealityBP/Config/DefaultScalability.ini
@@ -0,0 +1,90 @@
+
+[AntiAliasingQuality@0]
+r.PostProcessAAQuality=0
+
+[AntiAliasingQuality@1]
+r.PostProcessAAQuality=2
+
+[AntiAliasingQuality@2]
+r.PostProcessAAQuality=3
+
+[AntiAliasingQuality@3]
+r.PostProcessAAQuality=3
+
+;-----------------------------------------------------------------------------------------------------------------
+
+[PostProcessQuality@0]
+r.MotionBlurQuality=0
+r.AmbientOcclusionLevels=0
+r.LensFlareQuality=0
+r.SceneColorFringeQuality=0
+r.DepthOfFieldQuality=0
+r.BloomQuality=1
+r.FastBlurThreshold=0
+
+[PostProcessQuality@1]
+r.MotionBlurQuality=0
+r.AmbientOcclusionLevels=0
+r.LensFlareQuality=0
+r.SceneColorFringeQuality=0
+r.DepthOfFieldQuality=0
+r.BloomQuality=1
+r.FastBlurThreshold=0
+
+[PostProcessQuality@2]
+r.MotionBlurQuality=0
+r.AmbientOcclusionLevels=0
+r.LensFlareQuality=0
+r.SceneColorFringeQuality=0
+r.DepthOfFieldQuality=0
+r.BloomQuality=1
+r.FastBlurThreshold=0
+
+[PostProcessQuality@3]
+r.MotionBlurQuality=0
+r.AmbientOcclusionLevels=0
+r.LensFlareQuality=0
+r.SceneColorFringeQuality=0
+r.DepthOfFieldQuality=0
+r.BloomQuality=1
+r.FastBlurThreshold=0
+
+;-----------------------------------------------------------------------------------------------------------------
+
+[TextureQuality@0]
+r.MaxAnisotropy=0
+
+[TextureQuality@1]
+r.MaxAnisotropy=2
+
+[TextureQuality@2]
+r.MaxAnisotropy=4
+
+[TextureQuality@3]
+r.MaxAnisotropy=8
+
+;-----------------------------------------------------------------------------------------------------------------
+
+[EffectsQuality@0]
+r.TranslucencyVolumeBlur=0
+r.TranslucencyLightingVolumeDim=4
+r.SceneColorFormat=2
+r.SSR.Quality=0
+
+[EffectsQuality@1]
+r.TranslucencyVolumeBlur=0
+r.TranslucencyLightingVolumeDim=4
+r.SceneColorFormat=2
+r.SSR.Quality=0
+
+[EffectsQuality@2]
+r.TranslucencyVolumeBlur=0
+r.TranslucencyLightingVolumeDim=4
+r.SceneColorFormat=2
+r.SSR.Quality=0
+
+[EffectsQuality@3]
+r.TranslucencyVolumeBlur=0
+r.TranslucencyLightingVolumeDim=4
+r.SceneColorFormat=2
+r.SSR.Quality=0
\ No newline at end of file
