Fixed VR template warning due to missing default game mode

diff --git a/Templates/TP_VirtualRealityBP/Config/DefaultEngine.ini b/Templates/TP_VirtualRealityBP/Config/DefaultEngine.ini
--- a/Templates/TP_VirtualRealityBP/Config/DefaultEngine.ini
+++ b/Templates/TP_VirtualRealityBP/Config/DefaultEngine.ini
@@ -8,7 +8,7 @@ ThreePlayerSplitscreenLayout=FavorTop
 GameInstanceClass=/Script/Engine.GameInstance
 GameDefaultMap=/Game/VirtualRealityBP/Maps/MotionControllerMap.MotionControllerMap
 ServerDefaultMap=/Engine/Maps/Entry
-GlobalDefaultGameMode=None
+GlobalDefaultGameMode=/Script/Engine.GameModeBase
 GlobalDefaultServerGameMode=None
 
 [/Script/Engine.Engine]
