Pulling Dev-RPS

diff --git a/Templates/FP_FirstPerson/Source/FP_FirstPerson/FP_FirstPersonCharacter.cpp b/Templates/FP_FirstPerson/Source/FP_FirstPerson/FP_FirstPersonCharacter.cpp
--- a/Templates/FP_FirstPerson/Source/FP_FirstPerson/FP_FirstPersonCharacter.cpp
+++ b/Templates/FP_FirstPerson/Source/FP_FirstPerson/FP_FirstPersonCharacter.cpp
@@ -27,7 +27,7 @@ AFP_FirstPersonCharacter::AFP_FirstPersonCharacter()
 	// Create a CameraComponent	
 	FirstPersonCameraComponent = CreateDefaultSubobject<UCameraComponent>(TEXT("FirstPersonCamera"));
 	FirstPersonCameraComponent->SetupAttachment(GetCapsuleComponent());
-	FirstPersonCameraComponent->RelativeLocation = FVector(0, 0, 64.f); // Position the camera
+	FirstPersonCameraComponent->SetRelativeLocation(FVector(0, 0, 64.f)); // Position the camera
 	FirstPersonCameraComponent->bUsePawnControlRotation = true;
 	
 	// Create a mesh component that will be used when being viewed from a '1st person' view (when controlling this pawn)
@@ -246,7 +246,7 @@ void AFP_FirstPersonCharacter::LookUpAtRate(float Rate)
 FHitResult AFP_FirstPersonCharacter::WeaponTrace(const FVector& StartTrace, const FVector& EndTrace) const
 {
 	// Perform trace to retrieve hit info
-	FCollisionQueryParams TraceParams(SCENE_QUERY_STAT(WeaponTrace), true, Instigator);
+	FCollisionQueryParams TraceParams(SCENE_QUERY_STAT(WeaponTrace), true, GetInstigator());
 	TraceParams.bReturnPhysicalMaterial = true;
 
 	FHitResult Hit(ForceInit);
diff --git a/Templates/TP_2DSideScroller/Source/TP_2DSideScroller/TP_2DSideScrollerCharacter.cpp b/Templates/TP_2DSideScroller/Source/TP_2DSideScroller/TP_2DSideScrollerCharacter.cpp
--- a/Templates/TP_2DSideScroller/Source/TP_2DSideScroller/TP_2DSideScrollerCharacter.cpp
+++ b/Templates/TP_2DSideScroller/Source/TP_2DSideScroller/TP_2DSideScrollerCharacter.cpp
@@ -31,9 +31,9 @@ ATP_2DSideScrollerCharacter::ATP_2DSideScrollerCharacter()
 	CameraBoom->SetupAttachment(RootComponent);
 	CameraBoom->TargetArmLength = 500.0f;
 	CameraBoom->SocketOffset = FVector(0.0f, 0.0f, 75.0f);
-	CameraBoom->bAbsoluteRotation = true;
+	CameraBoom->SetUsingAbsoluteRotation(true);
 	CameraBoom->bDoCollisionTest = false;
-	CameraBoom->RelativeRotation = FRotator(0.0f, -90.0f, 0.0f);
+	CameraBoom->SetRelativeRotation(FRotator(0.0f, -90.0f, 0.0f));
 	
 
 	// Create an orthographic camera (no perspective) and attach it to the boom
@@ -43,7 +43,7 @@ ATP_2DSideScrollerCharacter::ATP_2DSideScrollerCharacter()
 	SideViewCameraComponent->SetupAttachment(CameraBoom, USpringArmComponent::SocketName);
 
 	// Prevent all automatic rotation behavior on the camera, character, and camera component
-	CameraBoom->bAbsoluteRotation = true;
+	CameraBoom->SetUsingAbsoluteRotation(true);
 	SideViewCameraComponent->bUsePawnControlRotation = false;
 	SideViewCameraComponent->bAutoActivate = true;
 	GetCharacterMovement()->bOrientRotationToMovement = false;
diff --git a/Templates/TP_AEC_BlankBP/Config/DefaultEngine.ini b/Templates/TP_AEC_BlankBP/Config/DefaultEngine.ini
--- a/Templates/TP_AEC_BlankBP/Config/DefaultEngine.ini
+++ b/Templates/TP_AEC_BlankBP/Config/DefaultEngine.ini
@@ -17,7 +17,7 @@ r.LightPropagationVolume=0
 ; quality settings
 r.ReflectionCaptureResolution=2048
 r.AllowStaticLighting=True
-r.HighResScreenshotDelay 8
+r.HighResScreenshotDelay=8
 r.DefaultBackBufferPixelFormat=4
 r.AllowGlobalClipPlane=False
 
@@ -37,8 +37,5 @@ EditorStartupMap=/Game/Main.Main
 [/Script/UnrealEd.UnrealEdEngine]
 ; Remove Engine Template maps
 -TemplateMapInfos=(ThumbnailTexture=Texture2D'/Engine/Maps/Templates/Thumbnails/Default.Default',Map="/Engine/Maps/Templates/Template_Default")
--TemplateMapInfos=(ThumbnailTexture=Texture2D'/Engine/Maps/Templates/Thumbnails/VR-Basic.VR-Basic',Map="/Engine/Maps/Templates/VR-Basic")
-
-; Add Predefined Template Maps that are suitable for this template
-;+TemplateMapInfos=(ThumbnailTexture=Texture2D'/Game/Developers/Maps/Thumbnails/AssetName.AssetName',Map="/Game/Maps/Templates/MapName")
-+TemplateMapInfos=(ThumbnailTexture=Texture2D'/Game/Developers/LevelTemplates/Thumbnails/Default.Default',Map="/Game/Developers/LevelTemplates/Default")
\ No newline at end of file
++TemplateMapInfos=(ThumbnailTexture=Texture2D'/Engine/Maps/Templates/Thumbnails/TimeOfDay.TimeOfDay',Map="/Engine/Maps/Templates/TimeOfDay_Default")
+-TemplateMapInfos=(ThumbnailTexture=Texture2D'/Engine/Maps/Templates/Thumbnails/VR-Basic.VR-Basic',Map="/Engine/Maps/Templates/VR-Basic")
\ No newline at end of file
diff --git a/Templates/TP_AEC_BlankBP/Config/TemplateDefs.ini b/Templates/TP_AEC_BlankBP/Config/TemplateDefs.ini
--- a/Templates/TP_AEC_BlankBP/Config/TemplateDefs.ini
+++ b/Templates/TP_AEC_BlankBP/Config/TemplateDefs.ini
@@ -6,14 +6,15 @@
 
 ; The categories that this template should appear in. The values in this list refer to the Key value given in TemplateCategories.ini.
 Categories=AEC
+Categories=MFG
 
 ; The localized display names for this template. Each entry specifies a Language code and Unicode text to display in the New Project dialog.
 LocalizedDisplayNames=(Language="en", Text="Blank")
 
 ; The localized descriptions for this template, which is displayed when the template is selected.
-LocalizedDescriptions=(Language="en", Text="Blank template with minimal functionality.")
+LocalizedDescriptions=(Language="en", Text="Blank template with minimal functionality and content.\n\nThis template enables you to use Datasmith to import models from various design applications and easily light them.")
 ; The sort key to use for the template to bypass the default alphabetical sorting.
-SortKey="_1"
+SortKey="0"
 
 ; A list of folders to avoid copying when creating a project from this template. 
 ; These should contain files only of concern to the template itself and not desired in the user's project.
@@ -81,7 +82,7 @@ HiddenSettings=Languages
 HiddenSettings=HardwareTarget
 HiddenSettings=GraphicsPreset 
 ;HiddenSettings=StarterContent 
-HiddenSettings=XR
+;HiddenSettings=XR
 
 ; The class types that this project uses.
 ; This is only for display in the template description, and does not affect project creation.
diff --git a/Templates/TP_AEC_BlankBP/TP_AEC_BlankBP.uproject b/Templates/TP_AEC_BlankBP/TP_AEC_BlankBP.uproject
--- a/Templates/TP_AEC_BlankBP/TP_AEC_BlankBP.uproject
+++ b/Templates/TP_AEC_BlankBP/TP_AEC_BlankBP.uproject
@@ -17,21 +17,32 @@
 			"Name": "DatasmithImporter",
 			"Enabled": true
 		},
+		{
+			"Name": "DatasmithCADImporter",
+			"Enabled": true,
+			"SupportedTargetPlatforms": [
+				"Win64"
+			]
+		},
 		{
 			"Name": "VariantManager",
-			"Enabled": True
+			"Enabled": true
 		},
 		{
 			"Name": "PythonScriptPlugin",
 			"Enabled": true
 		},
-				{
+		{
 			"Name": "HDRIBackdrop",
 			"Enabled": true
 		},
 		{
 			"Name": "SequencerScripting",
 			"Enabled": false
+		},
+				{
+			"Name": "StaticMeshEditorExtension",
+			"Enabled": true
 		}
 	]
 }
\ No newline at end of file
diff --git a/Templates/TP_ArchvisBP/Config/DefaultEngine.ini b/Templates/TP_ArchvisBP/Config/DefaultEngine.ini
new file mode 100644
--- /dev/null
+++ b/Templates/TP_ArchvisBP/Config/DefaultEngine.ini
@@ -0,0 +1,41 @@
+[/Script/HardwareTargeting.HardwareTargetingSettings]
+; quality settings
+TargetedHardwareClass=Desktop
+AppliedTargetedHardwareClass=Desktop
+DefaultGraphicsPerformance=Maximum
+AppliedDefaultGraphicsPerformance=Maximum
+
+[/Script/Engine.RendererSettings]
+; exposure control and tone mapping settings
+r.Tonemapper.Sharpen=2
+r.TemporalAA.Upsampling=True
+r.DefaultFeature.AutoExposure.ExtendDefaultLuminanceRange=True
+
+;dynamic gi settings
+r.LightPropagationVolume=0
+
+; quality settings
+r.ReflectionCaptureResolution=2048
+r.AllowStaticLighting=True
+r.HighResScreenshotDelay=8
+r.DefaultBackBufferPixelFormat=4
+r.AllowGlobalClipPlane=False
+
+; shader settings
+r.SupportSkyAtmosphereAffectsHeightFog=True
+r.DefaultFeature.LensFlare=True
+r.ClearCoatNormal=False
+r.NormalMapsForStaticLighting=False
+
+
+[/Script/EngineSettings.GameMapsSettings]
+;defines what level to be loaded by default in the editor and at playback time
+GameDefaultMap=/Game/ArchvisProject/Exterior.Exterior
+EditorStartupMap=/Game/ArchvisProject/Exterior.Exterior
+GlobalDefaultGameMode=/Game/ArchvisProject/ArchvisGameMode.ArchvisGameMode_C
+
+[/Script/UnrealEd.UnrealEdEngine]
+; Remove Engine Template maps
+-TemplateMapInfos=(ThumbnailTexture=Texture2D'/Engine/Maps/Templates/Thumbnails/Default.Default',Map="/Engine/Maps/Templates/Template_Default")
++TemplateMapInfos=(ThumbnailTexture=Texture2D'/Engine/Maps/Templates/Thumbnails/TimeOfDay.TimeOfDay',Map="/Engine/Maps/Templates/TimeOfDay_Default")
+-TemplateMapInfos=(ThumbnailTexture=Texture2D'/Engine/Maps/Templates/Thumbnails/VR-Basic.VR-Basic',Map="/Engine/Maps/Templates/VR-Basic")
diff --git a/Templates/TP_ArchvisBP/Config/DefaultLightmass.ini b/Templates/TP_ArchvisBP/Config/DefaultLightmass.ini
new file mode 100644
--- /dev/null
+++ b/Templates/TP_ArchvisBP/Config/DefaultLightmass.ini
@@ -0,0 +1,22 @@
+; These are tweaked defaults for various lightmass solver and export settings
+; Artist oriented lightmass settings are in the editor UI
+; Documentation for all of these is in UnrealLightmass / Public / SceneExport.h
+; This ini is reloaded every time a lighting build begins, no need to restart
+
+; Warning: overwriting this file with an old version will cause the editor to crash.  This file must be in sync with the editor executable.
+; Instead, create a DefaultLightmass.ini in your project and override just the values you need, then the overrides will continue to work on version upgrades.
+; https://docs.unrealengine.com/latest/INT/Programming/Basics/ConfigurationFiles/
+;
+; For example, in your project's Config/DefaultLightmass.ini:
+; [DevOptions.PrecomputedDynamicObjectLighting]
+; SurfaceLightSampleSpacing=400
+
+[DevOptions.StaticLightingSettings]
+bCompressLightmaps=False
+NumIndirectLightingBounces=5
+NumSkyLightingBounces=5
+IndirectLightingSmoothness=0.9
+IndirectLightingQuality=2.0
+
+[DevOptions.StaticLightingSceneConstants]
+StaticLightingLevelScale=0.5
\ No newline at end of file
diff --git a/Templates/TP_ArchvisBP/Config/TemplateDefs.ini b/Templates/TP_ArchvisBP/Config/TemplateDefs.ini
new file mode 100644
--- /dev/null
+++ b/Templates/TP_ArchvisBP/Config/TemplateDefs.ini
@@ -0,0 +1,95 @@
+; This file specifies how a template will be treated when creating a project from it.
+; It includes rules for ignoring files, replacing files and changing the contents of files.
+; This file is annotated to help with creating your own templates.
+
+[/Script/GameProjectGeneration.TemplateProjectDefs]
+
+; The categories that this template should appear in. The values in this list refer to the Key value given in TemplateCategories.ini.
+Categories=AEC
+
+; The localized display names for this template. Each entry specifies a Language code and Unicode text to display in the New Project dialog.
+LocalizedDisplayNames=(Language="en", Text="Archvis")
+
+; The localized descriptions for this template, which is displayed when the template is selected.
+LocalizedDescriptions=(Language="en", Text="This template is dedicated to architectural visualization workflows with examples for sun studies, interior renderings, and non-photorealistic stylized renderings.\n\nThis template enables you to use Datasmith to import models from various design applications and easily light them.")
+; The sort key to use for the template to bypass the default alphabetical sorting.
+SortKey="1"
+
+; A list of folders to avoid copying when creating a project from this template. 
+; These should contain files only of concern to the template itself and not desired in the user's project.
+FoldersToIgnore=Media
+
+; The files to ignore copying, relative to the template root folder. 
+; %TEMPLATENAME% will be replaced by the template's folder name when resolving these. 
+; Filename comparisons are case-insensitive. 
+; For example, if this template is "TP_Blank" and the user creates "MyProject", then "TP_Blank.uproject" will be copied to "MyProject.uproject".
+FilesToIgnore="%TEMPLATENAME%.uproject"
+FilesToIgnore="%TEMPLATENAME%.png"
+FilesToIgnore="Config/TemplateDefs.ini"
+FilesToIgnore="Manifest.json"
+FilesToIgnore="contents.txt"
+
+; NOTE: For the purposes of the below examples, this template is considered to be named "TP_Blank", and the user has created a project called "MyProject".
+
+; Folders that will be renamed when creating a project from this template.
+; %PROJECTNAME% will be replaced by the project name specified by the user during creation when resolving these rules.
+; For example,  then "Source/TP_Blank" will be copied to "Source/MyProject".
+FolderRenames=(From="Source/%TEMPLATENAME%", To="Source/%PROJECTNAME%")
+FolderRenames=(From="Source/%TEMPLATENAME%Editor", To="Source/%PROJECTNAME%Editor")
+
+; Rules for renaming files when creating a project from this template. 
+; %TEMPLATENAME_UPPERCASE% and %PROJECTNAME_UPPERCASE% are the template and project names in upper case, eg. for "TP_Blank", its upper case representation is "TP_BLANK".
+; %TEMPLATENAME_LOWERCASE% and %PROJECTNAME_LOWERCASE% are the template and project names in lower case, eg. for "TP_Blank", its lower case representation is "tp_blank".
+; A rule does not have to include any of these modifiers. 
+; Extensions is a list of files to consider for these replacements. Eg. ("cpp", "h") will try to apply a rule only to files like "SomeFile.cpp" and "SomeFile.h".
+; bCaseSensitive specifies whether this replacement should be case-sensitive.
+
+; Here are some default rules for renaming files:
+
+; This rule states that a file called "TP_BLANK.cpp" will be renamed to "MYPROJECT.cpp".
+FilenameReplacements=(Extensions=("cpp","h","ini","cs"), From="%TEMPLATENAME_UPPERCASE%", To="%PROJECTNAME_UPPERCASE%", bCaseSensitive=true)
+
+; This rule states that a file called "tp_blank.ini" will be renamed to "myproject.ini"
+FilenameReplacements=(Extensions=("cpp","h","ini","cs"), From="%TEMPLATENAME_LOWERCASE%", To="%PROJECTNAME_LOWERCASE%", bCaseSensitive=true)
+
+; And this rule catches any irregularly-cased files, so that "Tp_bLANK.ini" will be renamed to "MyProject.ini".
+FilenameReplacements=(Extensions=("cpp","h","ini","cs"), From="%TEMPLATENAME%",To="%PROJECTNAME%",bCaseSensitive=false)
+
+; Rules for renaming strings in the contents of files. 
+; The same rules as above apply, but any usages of strings that match these rules inside the file will be changed. This includes all substrings!
+; For example, "TP_Blank" will be replaced by "MyProject", but "TP_Blank_Name" will also be replaced by "MyProject_Name".
+ReplacementsInFiles=(Extensions=("cpp","h","ini","cs"), From="%TEMPLATENAME_UPPERCASE%", To="%PROJECTNAME_UPPERCASE%", bCaseSensitive=true)
+ReplacementsInFiles=(Extensions=("cpp","h","ini","cs"), From="%TEMPLATENAME_LOWERCASE%", To="%PROJECTNAME_LOWERCASE%", bCaseSensitive=true)
+ReplacementsInFiles=(Extensions=("cpp","h","ini","cs"), From="%TEMPLATENAME%", To="%PROJECTNAME%", bCaseSensitive=false)
+
+; The following settings are not enabled for this project, but are included for completeness:
+
+; Whether or not this template can be used to create a project. Default: true
+; If this is set to false, then the template becomes a Feature Pack that can only be referenced from other templates, or added to a project from the Content Browser.
+;bAllowProjectCreation=false
+
+; Which settings on the project settings page should be hidden for this template. Default: none
+; Allowed values are:
+; All - Hides the settings page entirely
+; Languages - Choose between C++ and Blueprint versions.
+; HardwareTarget - Choose between Mobile and Desktop hardware targets.
+; GraphicsPreset - Choose between Maximum and Scalable graphics.
+; StarterContent - Choose whether to add starter content to the project.
+; VR - Choose whether to enable Virtual Reality support for the project.
+;HiddenSettings=All
+HiddenSettings=Languages
+HiddenSettings=HardwareTarget
+HiddenSettings=GraphicsPreset 
+;HiddenSettings=StarterContent 
+;HiddenSettings=XR
+
+; The class types that this project uses.
+; This is only for display in the template description, and does not affect project creation.
+;ClassTypes=GameMode, Character, Actor
+
+; The asset types that this project uses.
+; This is only for display in the template description, and does not affect project creation.
+;AssetTypes=Skeleton, Static Mesh, Texture
+
+; Shared feature packs. The files in these packs listed in these structures marked as 'additionalfiles' will be copied on project generation.
+;SharedContentPacks=(MountName="2DSideScroller",DetailLevels=("Standard"))
diff --git a/Templates/TP_ArchvisBP/TP_ArchvisBP.uproject b/Templates/TP_ArchvisBP/TP_ArchvisBP.uproject
new file mode 100644
--- /dev/null
+++ b/Templates/TP_ArchvisBP/TP_ArchvisBP.uproject
@@ -0,0 +1,54 @@
+{
+	"FileVersion": 3,
+	"EngineAssociation": "",
+	"Category": "",
+	"Description": "",
+	"Plugins": [
+		{
+			"Name": "SunPosition",
+			"Enabled": true
+		},
+		{
+			"Name": "DataPrepEditor",
+			"Enabled": true
+		},
+		{
+			"Name": "DatasmithImporter",
+			"Enabled": true
+		},
+		{
+			"Name": "VariantManager",
+			"Enabled": true
+		},
+		{
+			"Name": "PythonScriptPlugin",
+			"Enabled": true
+		},
+		{
+			"Name": "HDRIBackdrop",
+			"Enabled": true
+		},
+		{
+			"Name": "SequencerScripting",
+			"Enabled": false
+		},
+		{
+			"Name": "DatasmithCADImporter",
+			"Enabled": true,
+			"SupportedTargetPlatforms": [
+				"Win64"
+			]
+		},
+		{
+			"Name": "DatasmithIFCImporter",
+			"Enabled": true,
+			"SupportedTargetPlatforms": [
+				"Win64"
+			]
+		},
+		{
+			"Name": "StaticMeshEditorExtension",
+			"Enabled": true
+		}
+	]
+}
\ No newline at end of file
diff --git a/Templates/TP_Blank/Config/TemplateDefs.ini b/Templates/TP_Blank/Config/TemplateDefs.ini
--- a/Templates/TP_Blank/Config/TemplateDefs.ini
+++ b/Templates/TP_Blank/Config/TemplateDefs.ini
@@ -20,6 +20,9 @@ SortKey="_1"
 ; These should contain files only of concern to the template itself and not desired in the user's project.
 FoldersToIgnore=Media
 
+; Is this a blank template? Blank templates do not specify a default map.
+bIsBlank=True
+
 ; The files to ignore copying, relative to the template root folder. 
 ; %TEMPLATENAME% will be replaced by the template's folder name when resolving these. 
 ; Filename comparisons are case-insensitive. 
@@ -82,13 +85,13 @@ ReplacementsInFiles=(Extensions=("cpp","h","ini","cs"), From="%TEMPLATENAME%", T
 ; HardwareTarget - Choose between Mobile and Desktop hardware targets.
 ; GraphicsPreset - Choose between Maximum and Scalable graphics.
 ; StarterContent - Choose whether to add starter content to the project.
-; VR - Choose whether to enable Virtual Reality support for the project.
+; XR - Choose whether to enable Mixed Reality support for the project.
 ;HiddenSettings=All
 ;HiddenSettings=Languages
 ;HiddenSettings=HardwareTarget
 ;HiddenSettings=GraphicsPreset 
 ;HiddenSettings=StarterContent 
-;HiddenSettings=VR 
+;HiddenSettings=XR 
 
 ; The class types that this project uses.
 ; This is only for display in the template description, and does not affect project creation.
diff --git a/Templates/TP_Blank/TP_Blank.uproject b/Templates/TP_Blank/TP_Blank.uproject
--- a/Templates/TP_Blank/TP_Blank.uproject
+++ b/Templates/TP_Blank/TP_Blank.uproject
@@ -3,5 +3,11 @@
 	"EngineAssociation": "",
 	"Category": "",
 	"Description": "",
-	"Modules": []
+	"Modules": [
+		{
+			"Name": "TP_Blank",
+			"Type": "Runtime",
+			"LoadingPhase": "Default"
+		}
+	]
 }
\ No newline at end of file
diff --git a/Templates/TP_BlankBP/Config/TemplateDefs.ini b/Templates/TP_BlankBP/Config/TemplateDefs.ini
--- a/Templates/TP_BlankBP/Config/TemplateDefs.ini
+++ b/Templates/TP_BlankBP/Config/TemplateDefs.ini
@@ -2,6 +2,8 @@
 
 SortKey="_1"
 
+bIsBlank=True
+
 Categories=Games
 
 LocalizedDisplayNames=(Language="en", Text="Blank")
@@ -31,7 +33,3 @@ FilenameReplacements=(Extensions=("cpp","h","ini","cs"), From="%TEMPLATENAME%",
 ReplacementsInFiles=(Extensions=("cpp","h","ini","cs"), From="%TEMPLATENAME_UPPERCASE%", To="%PROJECTNAME_UPPERCASE%", bCaseSensitive=true)
 ReplacementsInFiles=(Extensions=("cpp","h","ini","cs"), From="%TEMPLATENAME_LOWERCASE%", To="%PROJECTNAME_LOWERCASE%", bCaseSensitive=true)
 ReplacementsInFiles=(Extensions=("cpp","h","ini","cs"), From="%TEMPLATENAME%", To="%PROJECTNAME%", bCaseSensitive=false)
-
-SharedContentPacks=(MountName="Geometry",DetailLevels=("Standard","High"))
-
-
diff --git a/Templates/TP_CollaborativeBP/Config/DefaultEngine.ini b/Templates/TP_CollaborativeBP/Config/DefaultEngine.ini
new file mode 100644
--- /dev/null
+++ b/Templates/TP_CollaborativeBP/Config/DefaultEngine.ini
@@ -0,0 +1,99 @@
+[/Script/Engine.Engine]
++ActiveGameNameRedirects=(OldGameName="TP_BlankBP",NewGameName="/Script/MyProject3")
++ActiveGameNameRedirects=(OldGameName="/Script/TP_BlankBP",NewGameName="/Script/MyProject3")
+
+[/Script/HardwareTargeting.HardwareTargetingSettings]
+TargetedHardwareClass=Desktop
+AppliedTargetedHardwareClass=Desktop
+DefaultGraphicsPerformance=Maximum
+AppliedDefaultGraphicsPerformance=Maximum
+
+[/Script/EngineSettings.GameMapsSettings]
+EditorStartupMap=/Game/CollaborativeViewer/Maps/CollaborativeViewer_P.CollaborativeViewer_P
+GameDefaultMap=/Game/CollaborativeViewer/Maps/MainMenu.MainMenu
+GlobalDefaultGameMode=/Game/CollaborativeViewer/Blueprints/GameMode/BP_CollaborativeViewer_GameMode.BP_CollaborativeViewer_GameMode_C
+GameInstanceClass=/Game/CollaborativeViewer/Blueprints/GameInstance/BP_CollaborativeViewer_GameInstance.BP_CollaborativeViewer_GameInstance_C
+
+[/Script/Engine.CollisionProfile]
++Profiles=(Name="NoCollision",CollisionEnabled=NoCollision,ObjectTypeName="WorldStatic",CustomResponses=((Channel="Visibility",Response=ECR_Ignore),(Channel="Camera",Response=ECR_Ignore)),HelpMessage="No collision",bCanModify=False)
++Profiles=(Name="BlockAll",CollisionEnabled=QueryAndPhysics,ObjectTypeName="WorldStatic",CustomResponses=,HelpMessage="WorldStatic object that blocks all actors by default. All new custom channels will use its own default response. ",bCanModify=False)
++Profiles=(Name="OverlapAll",CollisionEnabled=QueryOnly,ObjectTypeName="WorldStatic",CustomResponses=((Channel="WorldStatic",Response=ECR_Overlap),(Channel="Pawn",Response=ECR_Overlap),(Channel="Visibility",Response=ECR_Overlap),(Channel="WorldDynamic",Response=ECR_Overlap),(Channel="Camera",Response=ECR_Overlap),(Channel="PhysicsBody",Response=ECR_Overlap),(Channel="Vehicle",Response=ECR_Overlap),(Channel="Destructible",Response=ECR_Overlap)),HelpMessage="WorldStatic object that overlaps all actors by default. All new custom channels will use its own default response. ",bCanModify=False)
++Profiles=(Name="BlockAllDynamic",CollisionEnabled=QueryAndPhysics,ObjectTypeName="WorldDynamic",CustomResponses=,HelpMessage="WorldDynamic object that blocks all actors by default. All new custom channels will use its own default response. ",bCanModify=False)
++Profiles=(Name="OverlapAllDynamic",CollisionEnabled=QueryOnly,ObjectTypeName="WorldDynamic",CustomResponses=((Channel="WorldStatic",Response=ECR_Overlap),(Channel="Pawn",Response=ECR_Overlap),(Channel="Visibility",Response=ECR_Overlap),(Channel="WorldDynamic",Response=ECR_Overlap),(Channel="Camera",Response=ECR_Overlap),(Channel="PhysicsBody",Response=ECR_Overlap),(Channel="Vehicle",Response=ECR_Overlap),(Channel="Destructible",Response=ECR_Overlap)),HelpMessage="WorldDynamic object that overlaps all actors by default. All new custom channels will use its own default response. ",bCanModify=False)
++Profiles=(Name="IgnoreOnlyPawn",CollisionEnabled=QueryOnly,ObjectTypeName="WorldDynamic",CustomResponses=((Channel="Pawn",Response=ECR_Ignore),(Channel="Vehicle",Response=ECR_Ignore)),HelpMessage="WorldDynamic object that ignores Pawn and Vehicle. All other channels will be set to default.",bCanModify=False)
++Profiles=(Name="OverlapOnlyPawn",CollisionEnabled=QueryOnly,ObjectTypeName="WorldDynamic",CustomResponses=((Channel="Pawn",Response=ECR_Overlap),(Channel="Vehicle",Response=ECR_Overlap),(Channel="Camera",Response=ECR_Ignore)),HelpMessage="WorldDynamic object that overlaps Pawn, Camera, and Vehicle. All other channels will be set to default. ",bCanModify=False)
++Profiles=(Name="Pawn",CollisionEnabled=QueryAndPhysics,ObjectTypeName="Pawn",CustomResponses=((Channel="Visibility",Response=ECR_Ignore)),HelpMessage="Pawn object. Can be used for capsule of any playerable character or AI. ",bCanModify=False)
++Profiles=(Name="CharacterMesh",CollisionEnabled=QueryOnly,ObjectTypeName="Pawn",CustomResponses=((Channel="Pawn",Response=ECR_Ignore),(Channel="Vehicle",Response=ECR_Ignore),(Channel="Visibility",Response=ECR_Ignore)),HelpMessage="Pawn object that is used for Character Mesh. All other channels will be set to default.",bCanModify=False)
++Profiles=(Name="PhysicsActor",CollisionEnabled=QueryAndPhysics,ObjectTypeName="PhysicsBody",CustomResponses=,HelpMessage="Simulating actors",bCanModify=False)
++Profiles=(Name="Destructible",CollisionEnabled=QueryAndPhysics,ObjectTypeName="Destructible",CustomResponses=,HelpMessage="Destructible actors",bCanModify=False)
++Profiles=(Name="InvisibleWall",CollisionEnabled=QueryAndPhysics,ObjectTypeName="WorldStatic",CustomResponses=((Channel="Visibility",Response=ECR_Ignore)),HelpMessage="WorldStatic object that is invisible.",bCanModify=False)
++Profiles=(Name="InvisibleWallDynamic",CollisionEnabled=QueryAndPhysics,ObjectTypeName="WorldDynamic",CustomResponses=((Channel="Visibility",Response=ECR_Ignore)),HelpMessage="WorldDynamic object that is invisible.",bCanModify=False)
++Profiles=(Name="Trigger",CollisionEnabled=QueryOnly,ObjectTypeName="WorldDynamic",CustomResponses=((Channel="WorldStatic",Response=ECR_Overlap),(Channel="Pawn",Response=ECR_Overlap),(Channel="Visibility",Response=ECR_Ignore),(Channel="WorldDynamic",Response=ECR_Overlap),(Channel="Camera",Response=ECR_Overlap),(Channel="PhysicsBody",Response=ECR_Overlap),(Channel="Vehicle",Response=ECR_Overlap),(Channel="Destructible",Response=ECR_Overlap)),HelpMessage="WorldDynamic object that is used for trigger. All other channels will be set to default.",bCanModify=False)
++Profiles=(Name="Ragdoll",CollisionEnabled=QueryAndPhysics,ObjectTypeName="PhysicsBody",CustomResponses=((Channel="Pawn",Response=ECR_Ignore),(Channel="Visibility",Response=ECR_Ignore)),HelpMessage="Simulating Skeletal Mesh Component. All other channels will be set to default.",bCanModify=False)
++Profiles=(Name="Vehicle",CollisionEnabled=QueryAndPhysics,ObjectTypeName="Vehicle",CustomResponses=,HelpMessage="Vehicle object that blocks Vehicle, WorldStatic, and WorldDynamic. All other channels will be set to default.",bCanModify=False)
++Profiles=(Name="Spectator",CollisionEnabled=QueryOnly,ObjectTypeName="Pawn",CustomResponses=((Channel="WorldStatic"),(Channel="Pawn",Response=ECR_Ignore),(Channel="Visibility",Response=ECR_Ignore),(Channel="WorldDynamic",Response=ECR_Ignore),(Channel="Camera",Response=ECR_Ignore),(Channel="PhysicsBody",Response=ECR_Ignore),(Channel="Vehicle",Response=ECR_Ignore),(Channel="Destructible",Response=ECR_Ignore)),HelpMessage="Pawn object that ignores all other actors except WorldStatic.",bCanModify=False)
++Profiles=(Name="UI",CollisionEnabled=QueryOnly,ObjectTypeName="WorldDynamic",CustomResponses=((Channel="WorldStatic",Response=ECR_Overlap),(Channel="Pawn",Response=ECR_Overlap),(Channel="Visibility"),(Channel="WorldDynamic",Response=ECR_Overlap),(Channel="Camera",Response=ECR_Overlap),(Channel="PhysicsBody",Response=ECR_Overlap),(Channel="Vehicle",Response=ECR_Overlap),(Channel="Destructible",Response=ECR_Overlap)),HelpMessage="WorldStatic object that overlaps all actors by default. All new custom channels will use its own default response. ",bCanModify=False)
++DefaultChannelResponses=(Channel=ECC_GameTraceChannel1,Name="Annotation",DefaultResponse=ECR_Block,bTraceType=False,bStaticObject=False)
++DefaultChannelResponses=(Channel=ECC_GameTraceChannel2,Name="AnnotationStroke",DefaultResponse=ECR_Block,bTraceType=False,bStaticObject=False)
++DefaultChannelResponses=(Channel=ECC_GameTraceChannel3,Name="SnapShot",DefaultResponse=ECR_Block,bTraceType=False,bStaticObject=False)
++DefaultChannelResponses=(Channel=ECC_GameTraceChannel4,Name="DimenssionTool",DefaultResponse=ECR_Block,bTraceType=False,bStaticObject=False)
++ProfileRedirects=(OldName="BlockingVolume",NewName="InvisibleWall")
++ProfileRedirects=(OldName="InterpActor",NewName="IgnoreOnlyPawn")
++ProfileRedirects=(OldName="StaticMeshComponent",NewName="BlockAllDynamic")
++ProfileRedirects=(OldName="SkeletalMeshActor",NewName="PhysicsActor")
++ProfileRedirects=(OldName="InvisibleActor",NewName="InvisibleWallDynamic")
++CollisionChannelRedirects=(OldName="Static",NewName="WorldStatic")
++CollisionChannelRedirects=(OldName="Dynamic",NewName="WorldDynamic")
++CollisionChannelRedirects=(OldName="VehicleMovement",NewName="Vehicle")
++CollisionChannelRedirects=(OldName="PawnMovement",NewName="Pawn")
+
+[/Script/Engine.RendererSettings]
+r.CustomDepth=3
+
+[/Script/Engine.PhysicsSettings]
+DefaultGravityZ=-980.000000
+DefaultTerminalVelocity=4000.000000
+DefaultFluidFriction=0.300000
+SimulateScratchMemorySize=262144
+RagdollAggregateThreshold=4
+TriangleMeshTriangleMinAreaThreshold=5.000000
+bEnableShapeSharing=False
+bEnablePCM=True
+bEnableStabilization=False
+bWarnMissingLocks=True
+bEnable2DPhysics=False
+PhysicErrorCorrection=(PingExtrapolation=0.100000,PingLimit=100.000000,ErrorPerLinearDifference=1.000000,ErrorPerAngularDifference=1.000000,MaxRestoredStateError=1.000000,MaxLinearHardSnapDistance=400.000000,PositionLerp=0.000000,AngleLerp=0.400000,LinearVelocityCoefficient=100.000000,AngularVelocityCoefficient=10.000000,ErrorAccumulationSeconds=0.500000,ErrorAccumulationDistanceSq=15.000000,ErrorAccumulationSimilarity=100.000000)
+LockedAxis=Invalid
+DefaultDegreesOfFreedom=Full3D
+BounceThresholdVelocity=200.000000
+FrictionCombineMode=Average
+RestitutionCombineMode=Average
+MaxAngularVelocity=3600.000000
+MaxDepenetrationVelocity=0.000000
+ContactOffsetMultiplier=0.020000
+MinContactOffset=2.000000
+MaxContactOffset=8.000000
+bSimulateSkeletalMeshOnDedicatedServer=True
+DefaultShapeComplexity=CTF_UseSimpleAndComplex
+bDefaultHasComplexCollision=True
+bSuppressFaceRemapTable=False
+bSupportUVFromHitResults=False
+bDisableActiveActors=False
+bDisableKinematicStaticPairs=False
+bDisableKinematicKinematicPairs=False
+bDisableCCD=False
+bEnableEnhancedDeterminism=False
+AnimPhysicsMinDeltaTime=0.000000
+bSimulateAnimPhysicsAfterReset=False
+MaxPhysicsDeltaTime=0.033333
+bSubstepping=False
+bSubsteppingAsync=False
+MaxSubstepDeltaTime=0.016667
+MaxSubsteps=6
+SyncSceneSmoothingFactor=0.000000
+InitialAverageFrameRate=0.016667
+PhysXTreeRebuildRate=10
+DefaultBroadphaseSettings=(bUseMBPOnClient=False,bUseMBPOnServer=False,bUseMBPOuterBounds=False,MBPBounds=(Min=(X=0.000000,Y=0.000000,Z=0.000000),Max=(X=0.000000,Y=0.000000,Z=0.000000),IsValid=0),MBPOuterBounds=(Min=(X=0.000000,Y=0.000000,Z=0.000000),Max=(X=0.000000,Y=0.000000,Z=0.000000),IsValid=0),MBPNumSubdivs=2)
+ChaosSettings=(DefaultThreadingModel=TaskGraph,DedicatedThreadTickMode=VariableCappedWithTarget,DedicatedThreadBufferMode=Double)
+
+
diff --git a/Templates/TP_CollaborativeBP/Config/DefaultGame.ini b/Templates/TP_CollaborativeBP/Config/DefaultGame.ini
new file mode 100644
--- /dev/null
+++ b/Templates/TP_CollaborativeBP/Config/DefaultGame.ini
@@ -0,0 +1,7 @@
+
+
+
+[/Script/EngineSettings.GeneralProjectSettings]
+ProjectID=C684876B46D5306ADA6286AA20FA1808
+bStartInVR=False
+
diff --git a/Templates/TP_CollaborativeBP/Config/DefaultInput.ini b/Templates/TP_CollaborativeBP/Config/DefaultInput.ini
new file mode 100644
--- /dev/null
+++ b/Templates/TP_CollaborativeBP/Config/DefaultInput.ini
@@ -0,0 +1,115 @@
+
+[/Script/Engine.InputSettings]
++AxisConfig=(AxisKeyName="OculusTouch_Right_Thumbstick",AxisProperties=(DeadZone=0.000000,Sensitivity=1.000000,Exponent=1.000000,bInvert=False))
++AxisConfig=(AxisKeyName="OculusTouch_Right_FaceButton1",AxisProperties=(DeadZone=0.000000,Sensitivity=1.000000,Exponent=1.000000,bInvert=False))
++AxisConfig=(AxisKeyName="OculusTouch_Right_Trigger",AxisProperties=(DeadZone=0.000000,Sensitivity=1.000000,Exponent=1.000000,bInvert=False))
++AxisConfig=(AxisKeyName="OculusTouch_Right_FaceButton2",AxisProperties=(DeadZone=0.000000,Sensitivity=1.000000,Exponent=1.000000,bInvert=False))
++AxisConfig=(AxisKeyName="OculusTouch_Right_IndexPointing",AxisProperties=(DeadZone=0.000000,Sensitivity=1.000000,Exponent=1.000000,bInvert=False))
++AxisConfig=(AxisKeyName="OculusTouch_Right_ThumbUp",AxisProperties=(DeadZone=0.000000,Sensitivity=1.000000,Exponent=1.000000,bInvert=False))
++AxisConfig=(AxisKeyName="OculusTouchpad_Touchpad_X",AxisProperties=(DeadZone=0.000000,Sensitivity=1.000000,Exponent=1.000000,bInvert=False))
++AxisConfig=(AxisKeyName="OculusTouchpad_Touchpad_Y",AxisProperties=(DeadZone=0.000000,Sensitivity=1.000000,Exponent=1.000000,bInvert=False))
++AxisConfig=(AxisKeyName="SteamVR_Knuckles_Left_HandGrip",AxisProperties=(DeadZone=0.000000,Sensitivity=1.000000,Exponent=1.000000,bInvert=False))
++AxisConfig=(AxisKeyName="SteamVR_Knuckles_Left_IndexGrip",AxisProperties=(DeadZone=0.000000,Sensitivity=1.000000,Exponent=1.000000,bInvert=False))
++AxisConfig=(AxisKeyName="SteamVR_Knuckles_Left_MiddleGrip",AxisProperties=(DeadZone=0.000000,Sensitivity=1.000000,Exponent=1.000000,bInvert=False))
++AxisConfig=(AxisKeyName="SteamVR_Knuckles_Left_RingGrip",AxisProperties=(DeadZone=0.000000,Sensitivity=1.000000,Exponent=1.000000,bInvert=False))
++AxisConfig=(AxisKeyName="SteamVR_Knuckles_Left_PinkyGrip",AxisProperties=(DeadZone=0.000000,Sensitivity=1.000000,Exponent=1.000000,bInvert=False))
++AxisConfig=(AxisKeyName="SteamVR_Knuckles_Right_HandGrip",AxisProperties=(DeadZone=0.000000,Sensitivity=1.000000,Exponent=1.000000,bInvert=False))
++AxisConfig=(AxisKeyName="SteamVR_Knuckles_Right_IndexGrip",AxisProperties=(DeadZone=0.000000,Sensitivity=1.000000,Exponent=1.000000,bInvert=False))
++AxisConfig=(AxisKeyName="SteamVR_Knuckles_Right_MiddleGrip",AxisProperties=(DeadZone=0.000000,Sensitivity=1.000000,Exponent=1.000000,bInvert=False))
++AxisConfig=(AxisKeyName="SteamVR_Knuckles_Right_RingGrip",AxisProperties=(DeadZone=0.000000,Sensitivity=1.000000,Exponent=1.000000,bInvert=False))
++AxisConfig=(AxisKeyName="SteamVR_Knuckles_Right_PinkyGrip",AxisProperties=(DeadZone=0.000000,Sensitivity=1.000000,Exponent=1.000000,bInvert=False))
++AxisConfig=(AxisKeyName="Gamepad_LeftX",AxisProperties=(DeadZone=0.250000,Sensitivity=1.000000,Exponent=1.000000,bInvert=False))
++AxisConfig=(AxisKeyName="Gamepad_LeftY",AxisProperties=(DeadZone=0.250000,Sensitivity=1.000000,Exponent=1.000000,bInvert=False))
++AxisConfig=(AxisKeyName="Gamepad_RightX",AxisProperties=(DeadZone=0.250000,Sensitivity=1.000000,Exponent=1.000000,bInvert=False))
++AxisConfig=(AxisKeyName="Gamepad_RightY",AxisProperties=(DeadZone=0.250000,Sensitivity=1.000000,Exponent=1.000000,bInvert=False))
++AxisConfig=(AxisKeyName="MouseX",AxisProperties=(DeadZone=0.000000,Sensitivity=0.070000,Exponent=1.000000,bInvert=False))
++AxisConfig=(AxisKeyName="MouseY",AxisProperties=(DeadZone=0.000000,Sensitivity=0.070000,Exponent=1.000000,bInvert=False))
++AxisConfig=(AxisKeyName="MouseWheelAxis",AxisProperties=(DeadZone=0.000000,Sensitivity=1.000000,Exponent=1.000000,bInvert=False))
++AxisConfig=(AxisKeyName="Gamepad_LeftTriggerAxis",AxisProperties=(DeadZone=0.000000,Sensitivity=1.000000,Exponent=1.000000,bInvert=False))
++AxisConfig=(AxisKeyName="Gamepad_RightTriggerAxis",AxisProperties=(DeadZone=0.000000,Sensitivity=1.000000,Exponent=1.000000,bInvert=False))
++AxisConfig=(AxisKeyName="MotionController_Left_Thumbstick_X",AxisProperties=(DeadZone=0.000000,Sensitivity=1.000000,Exponent=1.000000,bInvert=False))
++AxisConfig=(AxisKeyName="MotionController_Left_Thumbstick_Y",AxisProperties=(DeadZone=0.000000,Sensitivity=1.000000,Exponent=1.000000,bInvert=False))
++AxisConfig=(AxisKeyName="MotionController_Left_TriggerAxis",AxisProperties=(DeadZone=0.000000,Sensitivity=1.000000,Exponent=1.000000,bInvert=False))
++AxisConfig=(AxisKeyName="MotionController_Left_Grip1Axis",AxisProperties=(DeadZone=0.000000,Sensitivity=1.000000,Exponent=1.000000,bInvert=False))
++AxisConfig=(AxisKeyName="MotionController_Left_Grip2Axis",AxisProperties=(DeadZone=0.000000,Sensitivity=1.000000,Exponent=1.000000,bInvert=False))
++AxisConfig=(AxisKeyName="MotionController_Right_Thumbstick_X",AxisProperties=(DeadZone=0.000000,Sensitivity=1.000000,Exponent=1.000000,bInvert=False))
++AxisConfig=(AxisKeyName="MotionController_Right_Thumbstick_Y",AxisProperties=(DeadZone=0.000000,Sensitivity=1.000000,Exponent=1.000000,bInvert=False))
++AxisConfig=(AxisKeyName="MotionController_Right_TriggerAxis",AxisProperties=(DeadZone=0.000000,Sensitivity=1.000000,Exponent=1.000000,bInvert=False))
++AxisConfig=(AxisKeyName="MotionController_Right_Grip1Axis",AxisProperties=(DeadZone=0.000000,Sensitivity=1.000000,Exponent=1.000000,bInvert=False))
++AxisConfig=(AxisKeyName="MotionController_Right_Grip2Axis",AxisProperties=(DeadZone=0.000000,Sensitivity=1.000000,Exponent=1.000000,bInvert=False))
++AxisConfig=(AxisKeyName="Gamepad_Special_Left_X",AxisProperties=(DeadZone=0.000000,Sensitivity=1.000000,Exponent=1.000000,bInvert=False))
++AxisConfig=(AxisKeyName="Gamepad_Special_Left_Y",AxisProperties=(DeadZone=0.000000,Sensitivity=1.000000,Exponent=1.000000,bInvert=False))
++AxisConfig=(AxisKeyName="MotionController_Left_Thumbstick_Z",AxisProperties=(DeadZone=0.000000,Sensitivity=1.000000,Exponent=1.000000,bInvert=False))
++AxisConfig=(AxisKeyName="MotionController_Right_Thumbstick_Z",AxisProperties=(DeadZone=0.000000,Sensitivity=1.000000,Exponent=1.000000,bInvert=False))
++AxisConfig=(AxisKeyName="OculusTouch_Left_Thumbstick",AxisProperties=(DeadZone=0.000000,Sensitivity=1.000000,Exponent=1.000000,bInvert=False))
++AxisConfig=(AxisKeyName="OculusTouch_Left_FaceButton1",AxisProperties=(DeadZone=0.000000,Sensitivity=1.000000,Exponent=1.000000,bInvert=False))
++AxisConfig=(AxisKeyName="OculusTouch_Left_Trigger",AxisProperties=(DeadZone=0.000000,Sensitivity=1.000000,Exponent=1.000000,bInvert=False))
++AxisConfig=(AxisKeyName="OculusTouch_Left_FaceButton2",AxisProperties=(DeadZone=0.000000,Sensitivity=1.000000,Exponent=1.000000,bInvert=False))
++AxisConfig=(AxisKeyName="OculusTouch_Left_IndexPointing",AxisProperties=(DeadZone=0.000000,Sensitivity=1.000000,Exponent=1.000000,bInvert=False))
++AxisConfig=(AxisKeyName="OculusTouch_Left_ThumbUp",AxisProperties=(DeadZone=0.000000,Sensitivity=1.000000,Exponent=1.000000,bInvert=False))
+bAltEnterTogglesFullscreen=True
+bF11TogglesFullscreen=True
+bUseMouseForTouch=False
+bEnableMouseSmoothing=True
+bEnableFOVScaling=True
+bCaptureMouseOnLaunch=True
+bDefaultViewportMouseLock=False
+bAlwaysShowTouchInterface=False
+bShowConsoleOnFourFingerTap=True
+bEnableGestureRecognizer=True
+bUseAutocorrect=False
+DefaultViewportMouseCaptureMode=CapturePermanently_IncludingInitialMouseDown
+DefaultViewportMouseLockMode=LockOnCapture
+FOVScale=0.011110
+DoubleClickTime=0.200000
++ActionMappings=(ActionName="MotionControllerL",bShift=False,bCtrl=False,bAlt=False,bCmd=False,Key=MotionController_Left_Trigger)
++ActionMappings=(ActionName="MotionControllerR",bShift=False,bCtrl=False,bAlt=False,bCmd=False,Key=MotionController_Right_Trigger)
++ActionMappings=(ActionName="PushAndPull",bShift=False,bCtrl=False,bAlt=False,bCmd=False,Key=MouseWheelAxis)
++ActionMappings=(ActionName="Quit",bShift=False,bCtrl=False,bAlt=False,bCmd=False,Key=Escape)
++ActionMappings=(ActionName="FlyMode",bShift=False,bCtrl=True,bAlt=False,bCmd=False,Key=One)
++ActionMappings=(ActionName="WalkMode",bShift=False,bCtrl=True,bAlt=False,bCmd=False,Key=Two)
++ActionMappings=(ActionName="OrbitMode",bShift=False,bCtrl=True,bAlt=False,bCmd=False,Key=Three)
++ActionMappings=(ActionName="VRMode",bShift=False,bCtrl=True,bAlt=False,bCmd=False,Key=Four)
++ActionMappings=(ActionName="ResetLocationL",bShift=False,bCtrl=False,bAlt=False,bCmd=False,Key=MotionController_Left_Grip1)
++ActionMappings=(ActionName="Bookmark_1",bShift=False,bCtrl=False,bAlt=False,bCmd=False,Key=One)
++ActionMappings=(ActionName="Bookmark_2",bShift=False,bCtrl=False,bAlt=False,bCmd=False,Key=Two)
++ActionMappings=(ActionName="Bookmark_3",bShift=False,bCtrl=False,bAlt=False,bCmd=False,Key=Three)
++ActionMappings=(ActionName="Bookmark_4",bShift=False,bCtrl=False,bAlt=False,bCmd=False,Key=Four)
++ActionMappings=(ActionName="Bookmark_5",bShift=False,bCtrl=False,bAlt=False,bCmd=False,Key=Five)
++ActionMappings=(ActionName="Bookmark_6",bShift=False,bCtrl=False,bAlt=False,bCmd=False,Key=Six)
++ActionMappings=(ActionName="Bookmark_7",bShift=False,bCtrl=False,bAlt=False,bCmd=False,Key=Seven)
++ActionMappings=(ActionName="Bookmark_8",bShift=False,bCtrl=False,bAlt=False,bCmd=False,Key=Eight)
++ActionMappings=(ActionName="Bookmark_9",bShift=False,bCtrl=False,bAlt=False,bCmd=False,Key=Nine)
++ActionMappings=(ActionName="OpenContextualMenu",bShift=False,bCtrl=False,bAlt=False,bCmd=False,Key=SpaceBar)
++ActionMappings=(ActionName="OrbitAroundObject",bShift=False,bCtrl=False,bAlt=False,bCmd=False,Key=RightMouseButton)
++ActionMappings=(ActionName="OrbitZoomOut",bShift=False,bCtrl=False,bAlt=False,bCmd=False,Key=MouseScrollDown)
++ActionMappings=(ActionName="OrbitZoomIn",bShift=False,bCtrl=False,bAlt=False,bCmd=False,Key=MouseScrollUp)
++ActionMappings=(ActionName="ShootLaser",bShift=False,bCtrl=False,bAlt=False,bCmd=False,Key=LeftMouseButton)
++ActionMappings=(ActionName="RotateView",bShift=False,bCtrl=False,bAlt=False,bCmd=False,Key=RightMouseButton)
++ActionMappings=(ActionName="SetPointOfinterest",bShift=False,bCtrl=False,bAlt=False,bCmd=False,Key=MiddleMouseButton)
++ActionMappings=(ActionName="LockToAxis",bShift=False,bCtrl=False,bAlt=False,bCmd=False,Key=LeftShift)
++AxisMappings=(AxisName="WalkForward",Scale=1.000000,Key=W)
++AxisMappings=(AxisName="WalkRight",Scale=1.000000,Key=D)
++AxisMappings=(AxisName="WalkForward",Scale=-1.000000,Key=S)
++AxisMappings=(AxisName="WalkRight",Scale=-1.000000,Key=A)
++AxisMappings=(AxisName="LookUp",Scale=1.000000,Key=MouseY)
++AxisMappings=(AxisName="LookRight",Scale=1.000000,Key=MouseX)
++AxisMappings=(AxisName="FlyUp",Scale=1.000000,Key=E)
++AxisMappings=(AxisName="FlyUp",Scale=-1.000000,Key=Q)
++AxisMappings=(AxisName="MotionControllerThumbLeft_Y",Scale=-1.000000,Key=MotionController_Left_Thumbstick_Y)
++AxisMappings=(AxisName="MotionControllerThumbRight_Y",Scale=-1.000000,Key=MotionController_Right_Thumbstick_Y)
++AxisMappings=(AxisName="MotionControllerThumbLeft_X",Scale=1.000000,Key=MotionController_Left_Thumbstick_X)
++AxisMappings=(AxisName="MotionControllerThumbRight_X",Scale=1.000000,Key=MotionController_Right_Thumbstick_X)
++AxisMappings=(AxisName="WalkForward",Scale=1.000000,Key=Gamepad_LeftY)
++AxisMappings=(AxisName="WalkRight",Scale=1.000000,Key=Gamepad_LeftX)
++AxisMappings=(AxisName="LookUp",Scale=-1.000000,Key=Gamepad_RightY)
++AxisMappings=(AxisName="LookRight",Scale=1.000000,Key=Gamepad_RightX)
++AxisMappings=(AxisName="pinch",Scale=1.000000,Key=Gesture_Pinch)
++AxisMappings=(AxisName="WalkForward",Scale=1.000000,Key=Up)
++AxisMappings=(AxisName="WalkForward",Scale=-1.000000,Key=Down)
++AxisMappings=(AxisName="WalkRight",Scale=1.000000,Key=Right)
++AxisMappings=(AxisName="WalkRight",Scale=-1.000000,Key=Left)
+DefaultTouchInterface=/Game/CollaborativeViewer/Blueprints/Pawn/Touch/TouchInterface.TouchInterface
+ConsoleKey=None
++ConsoleKeys=Tilde
+
diff --git a/Templates/TP_CollaborativeBP/Config/IOS/IOSEngine.ini b/Templates/TP_CollaborativeBP/Config/IOS/IOSEngine.ini
new file mode 100644
--- /dev/null
+++ b/Templates/TP_CollaborativeBP/Config/IOS/IOSEngine.ini
@@ -0,0 +1,8 @@
+[/Script/IOSRuntimeSettings.IOSRuntimeSettings]
+;AudioSampleRate=48000
+AudioMaxChannels=16
+;AudioCallbackBufferFrameSize=1024
+;AudioNumBuffersToEnqueue=2
+;AudioNumSourceWorkers=0
+[OnlineSubsystem]
+DefaultPlatformService=Null
\ No newline at end of file
diff --git a/Templates/TP_CollaborativeBP/Config/TemplateDefs.ini b/Templates/TP_CollaborativeBP/Config/TemplateDefs.ini
new file mode 100644
--- /dev/null
+++ b/Templates/TP_CollaborativeBP/Config/TemplateDefs.ini
@@ -0,0 +1,97 @@
+﻿; This file specifies how a template will be treated when creating a project from it.
+; It includes rules for ignoring files, replacing files and changing the contents of files.
+; This file is annotated to help with creating your own templates.
+
+[/Script/GameProjectGeneration.TemplateProjectDefs]
+
+; The categories that this template should appear in. The values in this list refer to the Key value given in TemplateCategories.ini.
+Categories=AEC
+Categories=MFG
+
+; The localized display names for this template. Each entry specifies a Language code and Unicode text to display in the New Project dialog.
+LocalizedDisplayNames=(Language="en", Text="Collab Viewer")
+
+; The localized descriptions for this template, which is displayed when the template is selected.
+LocalizedDescriptions=(Language="en", Text="This project template provides navigation and interaction for various platforms (Desktop, VR). Use as a starting point for collaborative industrial or architectural projects.\n\nThis template enables you to use Datasmith to import models from various design applications and easily light them.")
+; The sort key to use for the template to bypass the default alphabetical sorting.
+SortKey="3"
+
+; A list of folders to avoid copying when creating a project from this template. 
+; These should contain files only of concern to the template itself and not desired in the user's project.
+FoldersToIgnore=Media
+
+; The files to ignore copying, relative to the template root folder. 
+; %TEMPLATENAME% will be replaced by the template's folder name when resolving these. 
+; Filename comparisons are case-insensitive. 
+; For example, if this template is "TP_Blank" and the user creates "MyProject", then "TP_Blank.uproject" will be copied to "MyProject.uproject".
+FilesToIgnore="%TEMPLATENAME%.uproject"
+FilesToIgnore="%TEMPLATENAME%.png"
+FilesToIgnore="Config/TemplateDefs.ini"
+FilesToIgnore="Manifest.json"
+FilesToIgnore="contents.txt"
+
+; NOTE: For the purposes of the below examples, this template is considered to be named "TP_Blank", and the user has created a project called "MyProject".
+
+; Folders that will be renamed when creating a project from this template.
+; %PROJECTNAME% will be replaced by the project name specified by the user during creation when resolving these rules.
+; For example,  then "Source/TP_Blank" will be copied to "Source/MyProject".
+FolderRenames=(From="Source/%TEMPLATENAME%", To="Source/%PROJECTNAME%")
+FolderRenames=(From="Source/%TEMPLATENAME%Editor", To="Source/%PROJECTNAME%Editor")
+
+; Rules for renaming files when creating a project from this template. 
+; %TEMPLATENAME_UPPERCASE% and %PROJECTNAME_UPPERCASE% are the template and project names in upper case, eg. for "TP_Blank", its upper case representation is "TP_BLANK".
+; %TEMPLATENAME_LOWERCASE% and %PROJECTNAME_LOWERCASE% are the template and project names in lower case, eg. for "TP_Blank", its lower case representation is "tp_blank".
+; A rule does not have to include any of these modifiers. 
+; Extensions is a list of files to consider for these replacements. Eg. ("cpp", "h") will try to apply a rule only to files like "SomeFile.cpp" and "SomeFile.h".
+; bCaseSensitive specifies whether this replacement should be case-sensitive.
+
+; Here are some default rules for renaming files:
+
+; This rule states that a file called "TP_BLANK.cpp" will be renamed to "MYPROJECT.cpp".
+FilenameReplacements=(Extensions=("cpp","h","ini","cs"), From="%TEMPLATENAME_UPPERCASE%", To="%PROJECTNAME_UPPERCASE%", bCaseSensitive=true)
+
+; This rule states that a file called "tp_blank.ini" will be renamed to "myproject.ini"
+FilenameReplacements=(Extensions=("cpp","h","ini","cs"), From="%TEMPLATENAME_LOWERCASE%", To="%PROJECTNAME_LOWERCASE%", bCaseSensitive=true)
+
+; And this rule catches any irregularly-cased files, so that "Tp_bLANK.ini" will be renamed to "MyProject.ini".
+FilenameReplacements=(Extensions=("cpp","h","ini","cs"), From="%TEMPLATENAME%",To="%PROJECTNAME%",bCaseSensitive=false)
+
+; Rules for renaming strings in the contents of files. 
+; The same rules as above apply, but any usages of strings that match these rules inside the file will be changed. This includes all substrings!
+; For example, "TP_Blank" will be replaced by "MyProject", but "TP_Blank_Name" will also be replaced by "MyProject_Name".
+ReplacementsInFiles=(Extensions=("cpp","h","ini","cs"), From="%TEMPLATENAME_UPPERCASE%", To="%PROJECTNAME_UPPERCASE%", bCaseSensitive=true)
+ReplacementsInFiles=(Extensions=("cpp","h","ini","cs"), From="%TEMPLATENAME_LOWERCASE%", To="%PROJECTNAME_LOWERCASE%", bCaseSensitive=true)
+ReplacementsInFiles=(Extensions=("cpp","h","ini","cs"), From="%TEMPLATENAME%", To="%PROJECTNAME%", bCaseSensitive=false)
+
+; The following settings are not enabled for this project, but are included for completeness:
+
+; Whether or not this template can be used to create a project. Default: true
+; If this is set to false, then the template becomes a Feature Pack that can only be referenced from other templates, or added to a project from the Content Browser.
+;bAllowProjectCreation=false
+
+; Which settings on the project settings page should be hidden for this template. Default: none
+; Allowed values are:
+; All - Hides the settings page entirely
+; Languages - Choose between C++ and Blueprint versions.
+; HardwareTarget - Choose between Mobile and Desktop hardware targets.
+; GraphicsPreset - Choose between Maximum and Scalable graphics.
+; StarterContent - Choose whether to add starter content to the project.
+; XR - Choose whether to enable Virtual Reality support for the project.
+;HiddenSettings=All
+HiddenSettings=Languages
+HiddenSettings=HardwareTarget
+HiddenSettings=GraphicsPreset 
+;HiddenSettings=StarterContent 
+HiddenSettings=XR
+
+
+; The class types that this project uses.
+; This is only for display in the template description, and does not affect project creation.
+;ClassTypes=GameMode, Character, Actor
+
+; The asset types that this project uses.
+; This is only for display in the template description, and does not affect project creation.
+;AssetTypes=Skeleton, Static Mesh, Texture
+
+; Shared feature packs. The files in these packs listed in these structures marked as 'additionalfiles' will be copied on project generation.
+;SharedContentPacks=(MountName="2DSideScroller",DetailLevels=("Standard"))
diff --git a/Templates/TP_CollaborativeBP/Config/config.ini b/Templates/TP_CollaborativeBP/Config/config.ini
new file mode 100644
--- /dev/null
+++ b/Templates/TP_CollaborativeBP/Config/config.ini
@@ -0,0 +1,119 @@
+[InputSettings]
++AxisConfig=(AxisKeyName="OculusTouch_Right_Thumbstick",AxisProperties=(DeadZone=0.000000,Sensitivity=1.000000,Exponent=1.000000,bInvert=False))
++AxisConfig=(AxisKeyName="OculusTouch_Right_FaceButton1",AxisProperties=(DeadZone=0.000000,Sensitivity=1.000000,Exponent=1.000000,bInvert=False))
++AxisConfig=(AxisKeyName="OculusTouch_Right_Trigger",AxisProperties=(DeadZone=0.000000,Sensitivity=1.000000,Exponent=1.000000,bInvert=False))
++AxisConfig=(AxisKeyName="OculusTouch_Right_FaceButton2",AxisProperties=(DeadZone=0.000000,Sensitivity=1.000000,Exponent=1.000000,bInvert=False))
++AxisConfig=(AxisKeyName="OculusTouch_Right_IndexPointing",AxisProperties=(DeadZone=0.000000,Sensitivity=1.000000,Exponent=1.000000,bInvert=False))
++AxisConfig=(AxisKeyName="OculusTouch_Right_ThumbUp",AxisProperties=(DeadZone=0.000000,Sensitivity=1.000000,Exponent=1.000000,bInvert=False))
++AxisConfig=(AxisKeyName="OculusTouchpad_Touchpad_X",AxisProperties=(DeadZone=0.000000,Sensitivity=1.000000,Exponent=1.000000,bInvert=False))
++AxisConfig=(AxisKeyName="OculusTouchpad_Touchpad_Y",AxisProperties=(DeadZone=0.000000,Sensitivity=1.000000,Exponent=1.000000,bInvert=False))
++AxisConfig=(AxisKeyName="SteamVR_Knuckles_Left_HandGrip",AxisProperties=(DeadZone=0.000000,Sensitivity=1.000000,Exponent=1.000000,bInvert=False))
++AxisConfig=(AxisKeyName="SteamVR_Knuckles_Left_IndexGrip",AxisProperties=(DeadZone=0.000000,Sensitivity=1.000000,Exponent=1.000000,bInvert=False))
++AxisConfig=(AxisKeyName="SteamVR_Knuckles_Left_MiddleGrip",AxisProperties=(DeadZone=0.000000,Sensitivity=1.000000,Exponent=1.000000,bInvert=False))
++AxisConfig=(AxisKeyName="SteamVR_Knuckles_Left_RingGrip",AxisProperties=(DeadZone=0.000000,Sensitivity=1.000000,Exponent=1.000000,bInvert=False))
++AxisConfig=(AxisKeyName="SteamVR_Knuckles_Left_PinkyGrip",AxisProperties=(DeadZone=0.000000,Sensitivity=1.000000,Exponent=1.000000,bInvert=False))
++AxisConfig=(AxisKeyName="SteamVR_Knuckles_Right_HandGrip",AxisProperties=(DeadZone=0.000000,Sensitivity=1.000000,Exponent=1.000000,bInvert=False))
++AxisConfig=(AxisKeyName="SteamVR_Knuckles_Right_IndexGrip",AxisProperties=(DeadZone=0.000000,Sensitivity=1.000000,Exponent=1.000000,bInvert=False))
++AxisConfig=(AxisKeyName="SteamVR_Knuckles_Right_MiddleGrip",AxisProperties=(DeadZone=0.000000,Sensitivity=1.000000,Exponent=1.000000,bInvert=False))
++AxisConfig=(AxisKeyName="SteamVR_Knuckles_Right_RingGrip",AxisProperties=(DeadZone=0.000000,Sensitivity=1.000000,Exponent=1.000000,bInvert=False))
++AxisConfig=(AxisKeyName="SteamVR_Knuckles_Right_PinkyGrip",AxisProperties=(DeadZone=0.000000,Sensitivity=1.000000,Exponent=1.000000,bInvert=False))
++AxisConfig=(AxisKeyName="Gamepad_LeftX",AxisProperties=(DeadZone=0.250000,Sensitivity=1.000000,Exponent=1.000000,bInvert=False))
++AxisConfig=(AxisKeyName="Gamepad_LeftY",AxisProperties=(DeadZone=0.250000,Sensitivity=1.000000,Exponent=1.000000,bInvert=False))
++AxisConfig=(AxisKeyName="Gamepad_RightX",AxisProperties=(DeadZone=0.250000,Sensitivity=1.000000,Exponent=1.000000,bInvert=False))
++AxisConfig=(AxisKeyName="Gamepad_RightY",AxisProperties=(DeadZone=0.250000,Sensitivity=1.000000,Exponent=1.000000,bInvert=False))
++AxisConfig=(AxisKeyName="MouseX",AxisProperties=(DeadZone=0.000000,Sensitivity=0.070000,Exponent=1.000000,bInvert=False))
++AxisConfig=(AxisKeyName="MouseY",AxisProperties=(DeadZone=0.000000,Sensitivity=0.070000,Exponent=1.000000,bInvert=False))
++AxisConfig=(AxisKeyName="MouseWheelAxis",AxisProperties=(DeadZone=0.000000,Sensitivity=1.000000,Exponent=1.000000,bInvert=False))
++AxisConfig=(AxisKeyName="Gamepad_LeftTriggerAxis",AxisProperties=(DeadZone=0.000000,Sensitivity=1.000000,Exponent=1.000000,bInvert=False))
++AxisConfig=(AxisKeyName="Gamepad_RightTriggerAxis",AxisProperties=(DeadZone=0.000000,Sensitivity=1.000000,Exponent=1.000000,bInvert=False))
++AxisConfig=(AxisKeyName="MotionController_Left_Thumbstick_X",AxisProperties=(DeadZone=0.000000,Sensitivity=1.000000,Exponent=1.000000,bInvert=False))
++AxisConfig=(AxisKeyName="MotionController_Left_Thumbstick_Y",AxisProperties=(DeadZone=0.000000,Sensitivity=1.000000,Exponent=1.000000,bInvert=False))
++AxisConfig=(AxisKeyName="MotionController_Left_TriggerAxis",AxisProperties=(DeadZone=0.000000,Sensitivity=1.000000,Exponent=1.000000,bInvert=False))
++AxisConfig=(AxisKeyName="MotionController_Left_Grip1Axis",AxisProperties=(DeadZone=0.000000,Sensitivity=1.000000,Exponent=1.000000,bInvert=False))
++AxisConfig=(AxisKeyName="MotionController_Left_Grip2Axis",AxisProperties=(DeadZone=0.000000,Sensitivity=1.000000,Exponent=1.000000,bInvert=False))
++AxisConfig=(AxisKeyName="MotionController_Right_Thumbstick_X",AxisProperties=(DeadZone=0.000000,Sensitivity=1.000000,Exponent=1.000000,bInvert=False))
++AxisConfig=(AxisKeyName="MotionController_Right_Thumbstick_Y",AxisProperties=(DeadZone=0.000000,Sensitivity=1.000000,Exponent=1.000000,bInvert=False))
++AxisConfig=(AxisKeyName="MotionController_Right_TriggerAxis",AxisProperties=(DeadZone=0.000000,Sensitivity=1.000000,Exponent=1.000000,bInvert=False))
++AxisConfig=(AxisKeyName="MotionController_Right_Grip1Axis",AxisProperties=(DeadZone=0.000000,Sensitivity=1.000000,Exponent=1.000000,bInvert=False))
++AxisConfig=(AxisKeyName="MotionController_Right_Grip2Axis",AxisProperties=(DeadZone=0.000000,Sensitivity=1.000000,Exponent=1.000000,bInvert=False))
++AxisConfig=(AxisKeyName="Gamepad_Special_Left_X",AxisProperties=(DeadZone=0.000000,Sensitivity=1.000000,Exponent=1.000000,bInvert=False))
++AxisConfig=(AxisKeyName="Gamepad_Special_Left_Y",AxisProperties=(DeadZone=0.000000,Sensitivity=1.000000,Exponent=1.000000,bInvert=False))
++AxisConfig=(AxisKeyName="MotionController_Left_Thumbstick_Z",AxisProperties=(DeadZone=0.000000,Sensitivity=1.000000,Exponent=1.000000,bInvert=False))
++AxisConfig=(AxisKeyName="MotionController_Right_Thumbstick_Z",AxisProperties=(DeadZone=0.000000,Sensitivity=1.000000,Exponent=1.000000,bInvert=False))
++AxisConfig=(AxisKeyName="OculusTouch_Left_Thumbstick",AxisProperties=(DeadZone=0.000000,Sensitivity=1.000000,Exponent=1.000000,bInvert=False))
++AxisConfig=(AxisKeyName="OculusTouch_Left_FaceButton1",AxisProperties=(DeadZone=0.000000,Sensitivity=1.000000,Exponent=1.000000,bInvert=False))
++AxisConfig=(AxisKeyName="OculusTouch_Left_Trigger",AxisProperties=(DeadZone=0.000000,Sensitivity=1.000000,Exponent=1.000000,bInvert=False))
++AxisConfig=(AxisKeyName="OculusTouch_Left_FaceButton2",AxisProperties=(DeadZone=0.000000,Sensitivity=1.000000,Exponent=1.000000,bInvert=False))
++AxisConfig=(AxisKeyName="OculusTouch_Left_IndexPointing",AxisProperties=(DeadZone=0.000000,Sensitivity=1.000000,Exponent=1.000000,bInvert=False))
++AxisConfig=(AxisKeyName="OculusTouch_Left_ThumbUp",AxisProperties=(DeadZone=0.000000,Sensitivity=1.000000,Exponent=1.000000,bInvert=False))
+bAltEnterTogglesFullscreen=True
+bF11TogglesFullscreen=True
+bUseMouseForTouch=False
+bEnableMouseSmoothing=True
+bEnableFOVScaling=True
+bCaptureMouseOnLaunch=True
+bDefaultViewportMouseLock=False
+bAlwaysShowTouchInterface=False
+bShowConsoleOnFourFingerTap=True
+bEnableGestureRecognizer=True
+bUseAutocorrect=False
+DefaultViewportMouseCaptureMode=CapturePermanently_IncludingInitialMouseDown
+DefaultViewportMouseLockMode=LockOnCapture
+FOVScale=0.011110
+DoubleClickTime=0.200000
++ActionMappings=(ActionName="MotionControllerL",bShift=False,bCtrl=False,bAlt=False,bCmd=False,Key=MotionController_Left_Trigger)
++ActionMappings=(ActionName="MotionControllerR",bShift=False,bCtrl=False,bAlt=False,bCmd=False,Key=MotionController_Right_Trigger)
++ActionMappings=(ActionName="PushAndPull",bShift=False,bCtrl=False,bAlt=False,bCmd=False,Key=MouseWheelAxis)
++ActionMappings=(ActionName="Quit",bShift=False,bCtrl=False,bAlt=False,bCmd=False,Key=Escape)
++ActionMappings=(ActionName="FlyMode",bShift=False,bCtrl=True,bAlt=False,bCmd=False,Key=One)
++ActionMappings=(ActionName="WalkMode",bShift=False,bCtrl=True,bAlt=False,bCmd=False,Key=Two)
++ActionMappings=(ActionName="OrbitMode",bShift=False,bCtrl=True,bAlt=False,bCmd=False,Key=Three)
++ActionMappings=(ActionName="VRMode",bShift=False,bCtrl=True,bAlt=False,bCmd=False,Key=Four)
++ActionMappings=(ActionName="ResetLocationL",bShift=False,bCtrl=False,bAlt=False,bCmd=False,Key=MotionController_Left_Grip1)
++ActionMappings=(ActionName="Bookmark_1",bShift=False,bCtrl=False,bAlt=False,bCmd=False,Key=One)
++ActionMappings=(ActionName="Bookmark_2",bShift=False,bCtrl=False,bAlt=False,bCmd=False,Key=Two)
++ActionMappings=(ActionName="Bookmark_3",bShift=False,bCtrl=False,bAlt=False,bCmd=False,Key=Three)
++ActionMappings=(ActionName="Bookmark_4",bShift=False,bCtrl=False,bAlt=False,bCmd=False,Key=Four)
++ActionMappings=(ActionName="Bookmark_5",bShift=False,bCtrl=False,bAlt=False,bCmd=False,Key=Five)
++ActionMappings=(ActionName="Bookmark_6",bShift=False,bCtrl=False,bAlt=False,bCmd=False,Key=Six)
++ActionMappings=(ActionName="Bookmark_7",bShift=False,bCtrl=False,bAlt=False,bCmd=False,Key=Seven)
++ActionMappings=(ActionName="Bookmark_8",bShift=False,bCtrl=False,bAlt=False,bCmd=False,Key=Eight)
++ActionMappings=(ActionName="Bookmark_9",bShift=False,bCtrl=False,bAlt=False,bCmd=False,Key=Nine)
++ActionMappings=(ActionName="OpenContextualMenu",bShift=False,bCtrl=False,bAlt=False,bCmd=False,Key=SpaceBar)
++ActionMappings=(ActionName="OrbitAroundObject",bShift=False,bCtrl=False,bAlt=False,bCmd=False,Key=RightMouseButton)
++ActionMappings=(ActionName="OrbitZoomOut",bShift=False,bCtrl=False,bAlt=False,bCmd=False,Key=MouseScrollDown)
++ActionMappings=(ActionName="OrbitZoomIn",bShift=False,bCtrl=False,bAlt=False,bCmd=False,Key=MouseScrollUp)
++ActionMappings=(ActionName="ShootLaser",bShift=False,bCtrl=False,bAlt=False,bCmd=False,Key=LeftMouseButton)
++ActionMappings=(ActionName="RotateView",bShift=False,bCtrl=False,bAlt=False,bCmd=False,Key=RightMouseButton)
++ActionMappings=(ActionName="SetPointOfinterest",bShift=False,bCtrl=False,bAlt=False,bCmd=False,Key=MiddleMouseButton)
++ActionMappings=(ActionName="LockToAxis",bShift=False,bCtrl=False,bAlt=False,bCmd=False,Key=LeftShift)
++AxisMappings=(AxisName="WalkForward",Scale=1.000000,Key=W)
++AxisMappings=(AxisName="WalkRight",Scale=1.000000,Key=D)
++AxisMappings=(AxisName="WalkForward",Scale=-1.000000,Key=S)
++AxisMappings=(AxisName="WalkRight",Scale=-1.000000,Key=A)
++AxisMappings=(AxisName="LookUp",Scale=1.000000,Key=MouseY)
++AxisMappings=(AxisName="LookRight",Scale=1.000000,Key=MouseX)
++AxisMappings=(AxisName="FlyUp",Scale=1.000000,Key=E)
++AxisMappings=(AxisName="FlyUp",Scale=-1.000000,Key=Q)
++AxisMappings=(AxisName="MotionControllerThumbLeft_Y",Scale=-1.000000,Key=MotionController_Left_Thumbstick_Y)
++AxisMappings=(AxisName="MotionControllerThumbRight_Y",Scale=-1.000000,Key=MotionController_Right_Thumbstick_Y)
++AxisMappings=(AxisName="MotionControllerThumbLeft_X",Scale=1.000000,Key=MotionController_Left_Thumbstick_X)
++AxisMappings=(AxisName="MotionControllerThumbRight_X",Scale=1.000000,Key=MotionController_Right_Thumbstick_X)
++AxisMappings=(AxisName="WalkForward",Scale=1.000000,Key=Gamepad_LeftY)
++AxisMappings=(AxisName="WalkRight",Scale=1.000000,Key=Gamepad_LeftX)
++AxisMappings=(AxisName="LookUp",Scale=-1.000000,Key=Gamepad_RightY)
++AxisMappings=(AxisName="LookRight",Scale=1.000000,Key=Gamepad_RightX)
++AxisMappings=(AxisName="pinch",Scale=1.000000,Key=Gesture_Pinch)
++AxisMappings=(AxisName="WalkForward",Scale=1.000000,Key=Up)
++AxisMappings=(AxisName="WalkForward",Scale=-1.000000,Key=Down)
++AxisMappings=(AxisName="WalkRight",Scale=1.000000,Key=Right)
++AxisMappings=(AxisName="WalkRight",Scale=-1.000000,Key=Left)
+DefaultTouchInterface=/Game/CollaborativeViewer/Blueprints/Pawn/Touch/TouchInterface.TouchInterface
+ConsoleKey=None
++ConsoleKeys=Tilde
+
+[Redirects]
+GameName=TP_CollaborativeBP
+
+[AdditionalFilesToAdd]
++Files=Enterprise/Templates/TP_CollaborativeBP/Content/*.*
diff --git a/Templates/TP_CollaborativeBP/TP_CollaborativeBP.uproject b/Templates/TP_CollaborativeBP/TP_CollaborativeBP.uproject
new file mode 100644
--- /dev/null
+++ b/Templates/TP_CollaborativeBP/TP_CollaborativeBP.uproject
@@ -0,0 +1,43 @@
+{
+	"FileVersion": 3,
+	"EngineAssociation": "",
+	"Category": "",
+	"Description": "",
+	"Plugins": [
+		{
+			"Name": "DataPrepEditor",
+			"Enabled": true
+		},
+		{
+			"Name": "DatasmithImporter",
+			"Enabled": true
+		},
+		{
+			"Name": "DatasmithCADImporter",
+			"Enabled": true,
+			"SupportedTargetPlatforms": [
+				"Win64"
+			]
+		},
+		{
+			"Name": "VariantManager",
+			"Enabled": true
+		},
+		{
+			"Name": "PythonScriptPlugin",
+			"Enabled": true
+		},
+		{
+			"Name": "HDRIBackdrop",
+			"Enabled": true
+		},
+		{
+			"Name": "SequencerScripting",
+			"Enabled": false
+		},
+		{
+			"Name": "StaticMeshEditorExtension",
+			"Enabled": true
+		}
+	]
+}
\ No newline at end of file
diff --git a/Templates/TP_CollaborativeBP/contents.txt b/Templates/TP_CollaborativeBP/contents.txt
new file mode 100644
--- /dev/null
+++ b/Templates/TP_CollaborativeBP/contents.txt
@@ -0,0 +1,4 @@
+../../../Templates/TP_CollaborativeBP/manifest.json
+../../../Templates/TP_CollaborativeBP/Config/config.ini
+../../../Templates/TP_CollaborativeBP/Media/TP_CollaborativeBP.png
+../../../Templates/TP_CollaborativeBP/Media/TP_CollaborativeBP_Preview.png
diff --git a/Templates/TP_CollaborativeBP/manifest.json b/Templates/TP_CollaborativeBP/manifest.json
new file mode 100644
--- /dev/null
+++ b/Templates/TP_CollaborativeBP/manifest.json
@@ -0,0 +1,31 @@
+{
+	"Version": 1,
+	"Name":
+	[
+		{
+			"Language": "en",
+			"Text": "Collab Viewer"
+		}
+	],
+	"Description":
+	[
+		{
+			"Language": "en",
+			"Text": "This project template provides navigation and interaction for various platforms (Desktop, VR). Use as a starting point for collaborative Industrial or Architectural projects."
+		}
+	],
+	"AssetTypes": [],
+	"SearchTags": [],
+	"ClassTypes": "",
+	"Category": "",
+	"FocusAsset": "",
+	"Thumbnail": "TP_CollaborativeBP.png",
+	"Screenshots":
+	[
+		"TP_CollaborativeBP_Preview.png"
+	],
+	"SortKey": "TP_CollaborativeBP",
+	"AdditionalFeaturePacks":
+	[
+	]
+}
diff --git a/Templates/TP_FirstPerson/Source/TP_FirstPerson/TP_FirstPersonCharacter.cpp b/Templates/TP_FirstPerson/Source/TP_FirstPerson/TP_FirstPersonCharacter.cpp
--- a/Templates/TP_FirstPerson/Source/TP_FirstPerson/TP_FirstPersonCharacter.cpp
+++ b/Templates/TP_FirstPerson/Source/TP_FirstPerson/TP_FirstPersonCharacter.cpp
@@ -29,7 +29,7 @@ ATP_FirstPersonCharacter::ATP_FirstPersonCharacter()
 	// Create a CameraComponent	
 	FirstPersonCameraComponent = CreateDefaultSubobject<UCameraComponent>(TEXT("FirstPersonCamera"));
 	FirstPersonCameraComponent->SetupAttachment(GetCapsuleComponent());
-	FirstPersonCameraComponent->RelativeLocation = FVector(-39.56f, 1.75f, 64.f); // Position the camera
+	FirstPersonCameraComponent->SetRelativeLocation(FVector(-39.56f, 1.75f, 64.f)); // Position the camera
 	FirstPersonCameraComponent->bUsePawnControlRotation = true;
 
 	// Create a mesh component that will be used when being viewed from a '1st person' view (when controlling this pawn)
@@ -38,8 +38,8 @@ ATP_FirstPersonCharacter::ATP_FirstPersonCharacter()
 	Mesh1P->SetupAttachment(FirstPersonCameraComponent);
 	Mesh1P->bCastDynamicShadow = false;
 	Mesh1P->CastShadow = false;
-	Mesh1P->RelativeRotation = FRotator(1.9f, -19.19f, 5.2f);
-	Mesh1P->RelativeLocation = FVector(-0.5f, -4.4f, -155.7f);
+	Mesh1P->SetRelativeRotation(FRotator(1.9f, -19.19f, 5.2f));
+	Mesh1P->SetRelativeLocation(FVector(-0.5f, -4.4f, -155.7f));
 
 	// Create a gun mesh component
 	FP_Gun = CreateDefaultSubobject<USkeletalMeshComponent>(TEXT("FP_Gun"));
diff --git a/Templates/TP_ME_BlankBP/Config/DefaultEngine.ini b/Templates/TP_ME_BlankBP/Config/DefaultEngine.ini
--- a/Templates/TP_ME_BlankBP/Config/DefaultEngine.ini
+++ b/Templates/TP_ME_BlankBP/Config/DefaultEngine.ini
@@ -5,40 +5,105 @@ AppliedTargetedHardwareClass=Desktop
 DefaultGraphicsPerformance=Maximum
 AppliedDefaultGraphicsPerformance=Maximum
 
+
 [/Script/Engine.RendererSettings]
 ; exposure control and tone mapping settings
 r.Tonemapper.Sharpen=2
 r.TemporalAA.Upsampling=True
 r.DefaultFeature.AutoExposure.ExtendDefaultLuminanceRange=True
+r.DefaultFeature.Bloom=False
+r.DefaultFeature.AutoExposure=False
+r.DefaultFeature.MotionBlur=False
+r.SceneRenderTargetResizeMethod=2
 
-;dynamic gi settings
+; dynamic gi settings
 r.LightPropagationVolume=0
 
 ; quality settings
 r.ReflectionCaptureResolution=2048
 r.AllowStaticLighting=True
-r.HighResScreenshotDelay 8
+r.HighResScreenshotDelay=8
 r.DefaultBackBufferPixelFormat=4
 r.AllowGlobalClipPlane=False
+r.GBufferFormat=3
 
 ; shader settings
+r.PostProcessing.PropagateAlpha=1
 r.SupportSkyAtmosphereAffectsHeightFog=True
 r.DefaultFeature.LensFlare=True
 r.ClearCoatNormal=False
 r.NormalMapsForStaticLighting=False
 
 
+
 [/Script/EngineSettings.GameMapsSettings]
 ;defines what level to be loaded by default in the editor and at playback time
-GameDefaultMap=/Game/Main.Main
-EditorStartupMap=/Game/Main.Main
-
+EditorStartupMap=/Game/Maps/Main.Main
+EditorStartupMap=/Game/Maps/Main.Main
 
 [/Script/UnrealEd.UnrealEdEngine]
 ; Remove Engine Template maps
--TemplateMapInfos=(ThumbnailTexture=Texture2D'/Engine/Maps/Templates/Thumbnails/Default.Default',Map="/Engine/Maps/Templates/Template_Default")
 -TemplateMapInfos=(ThumbnailTexture=Texture2D'/Engine/Maps/Templates/Thumbnails/VR-Basic.VR-Basic',Map="/Engine/Maps/Templates/VR-Basic")
 
-; Add Predefined Template Maps that are suitable for this template
-;+TemplateMapInfos=(ThumbnailTexture=Texture2D'/Game/Developers/Maps/Thumbnails/AssetName.AssetName',Map="/Game/Maps/Templates/MapName")
-+TemplateMapInfos=(ThumbnailTexture=Texture2D'/Game/Developers/LevelTemplates/Thumbnails/Default.Default',Map="/Game/Developers/LevelTemplates/Default")
\ No newline at end of file
+
+; Allows for Hardware Accelerated Video Decoding
+[/Script/WmfMediaFactory.WmfMediaSettings]
+AllowNonStandardCodecs=True
+HardwareAcceleratedVideoDecoding=True
+LowLatency=False
+NativeAudioOut=False
+
+; Adds Virtual Scouting Widget for VR Scouting
+
+[/Script/VPUtilitiesEditor.VPUtilitiesEditorSettings]
+VirtualScoutingUI=/VirtualProductionUtilities/Editor/VirtualScoutingWidget.VirtualScoutingWidget_C
+FlightSpeed=0.500000
+GripNavSpeed=0.250000
+bUseMetric=False
+bUseTransformGizmo=False
+bUseGripInertiaDamping=True
+InertiaDamping=0.950000
+bIsHelperSystemEnabled=True
+ScoutingSubsystemEdititorUtilityActorClassPath=/VirtualProductionUtilities/VirtualProductionHelpers.VirtualProductionHelpers_C
+
+
+; Enables WebControl API
+WebControl.EnableServerOnStartup=1
+
+;Adds Remote Session for Vcam
+[RemoteSession]
++Channels=(Name=FRemoteSessionFrameBufferChannel,Mode=Write)
++Channels=(Name=FRemoteSessionInputChannel,Mode=Read)
++Channels=(Name=FRemoteSessionXRTrackingChannel,Mode=Read)
+
+; Setup for Multiuser
+[/Script/Concert.ConcertClientConfig]
+bIsHeadless=False
+bInstallEditorToolbarButton=True
+bAutoConnect=False
+DefaultServerURL=
+DefaultSessionName=
+DefaultSessionToRestore=
+DefaultSaveSessionAs=
+ClientSettings=(DisplayName="",AvatarColor=(R=1.000000,G=1.000000,B=1.000000,A=1.000000),DesktopAvatarActorClass=/ConcertSyncClient/DesktopPresence.DesktopPresence_C,VRAvatarActorClass=/ConcertSyncClient/VRPresence.VRPresence_C,DiscoveryTimeoutSeconds=5,SessionTickFrequencySeconds=1,LatencyCompensationMs=0.000000,Tags=)
+EndpointSettings=(bEnableLogging=False,PurgeProcessedMessageDelaySeconds=30,RemoteEndpointTimeoutSeconds=60)
+
+; Setup for Ndisplay
+[/Script/Engine.Engine]
+GameEngine=/Script/DisplayCluster.DisplayClusterGameEngine
+UnrealEdEngine=/Script/DisplayClusterEditor.DisplayClusterEditorEngine
+GameViewportClientClassName=/Script/DisplayCluster.DisplayClusterViewportClient
+
+; Setup for Ndisplay
+[/Script/DisplayClusterEditor.DisplayClusterEditorSettings]
+bEnabled=True
+
+[/Script/WindowsTargetPlatform.WindowsTargetSettings]
+;Ndisplay Requirements
+DefaultGraphicsRHI=DefaultGraphicsRHI_DX11
++TargetedRHIs=PCD3D_SM5
+-TargetedRHIs=PCD3D_SM4
+-TargetedRHIs=GLSL_430
+-TargetedRHIs=GLSL_150
+
+
diff --git a/Templates/TP_ME_BlankBP/Config/DefaultGame.ini b/Templates/TP_ME_BlankBP/Config/DefaultGame.ini
--- a/Templates/TP_ME_BlankBP/Config/DefaultGame.ini
+++ b/Templates/TP_ME_BlankBP/Config/DefaultGame.ini
@@ -1 +1,11 @@
+[/Script/EngineSettings.GeneralProjectSettings]
+bUseBorderlessWindow=True
+
+[/Script/MediaFrameworkUtilities.MediaProfileSettings]
+bApplyInCommandlet=False
++MediaSourceProxy=/Game/VprodProject/MediaIO/MediaSource-01.MediaSource-01
++MediaSourceProxy=/Game/VprodProject/MediaIO/MediaSource-02.MediaSource-02
++MediaOutputProxy=/Game/VprodProject/MediaIO/MediaOutput-01.MediaOutput-01
+StartupMediaProfile=None
+
 
diff --git a/Templates/TP_ME_BlankBP/Config/DefaultVirtualProductionUtilities.ini b/Templates/TP_ME_BlankBP/Config/DefaultVirtualProductionUtilities.ini
new file mode 100644
--- /dev/null
+++ b/Templates/TP_ME_BlankBP/Config/DefaultVirtualProductionUtilities.ini
@@ -0,0 +1,13 @@
+
+[/Script/VPUtilitiesEditor.VPUtilitiesEditorSettings]
+VirtualScoutingUI=/VirtualProductionUtilities/Editor/VirtualScoutingWidget.VirtualScoutingWidget_C
+FlightSpeed=0.500000
+GripNavSpeed=0.250000
+bUseMetric=False
+bUseTransformGizmo=False
+bUseGripInertiaDamping=True
+InertiaDamping=0.950000
+bIsHelperSystemEnabled=True
+ScoutingSubsystemEdititorUtilityActorClassPath=/VirtualProductionUtilities/VirtualProductionHelpers.VirtualProductionHelpers_C
+
+
diff --git a/Templates/TP_ME_BlankBP/Config/TemplateDefs.ini b/Templates/TP_ME_BlankBP/Config/TemplateDefs.ini
--- a/Templates/TP_ME_BlankBP/Config/TemplateDefs.ini
+++ b/Templates/TP_ME_BlankBP/Config/TemplateDefs.ini
@@ -75,7 +75,7 @@ ReplacementsInFiles=(Extensions=("cpp","h","ini","cs"), From="%TEMPLATENAME%", T
 ; HardwareTarget - Choose between Mobile and Desktop hardware targets.
 ; GraphicsPreset - Choose between Maximum and Scalable graphics.
 ; StarterContent - Choose whether to add starter content to the project.
-; VR - Choose whether to enable Virtual Reality support for the project.
+; XR - Choose whether to enable Virtual Reality support for the project.
 ;HiddenSettings=All
 HiddenSettings=Languages
 HiddenSettings=HardwareTarget
diff --git a/Templates/TP_ME_BlankBP/TP_ME_BlankBP.uproject b/Templates/TP_ME_BlankBP/TP_ME_BlankBP.uproject
--- a/Templates/TP_ME_BlankBP/TP_ME_BlankBP.uproject
+++ b/Templates/TP_ME_BlankBP/TP_ME_BlankBP.uproject
@@ -3,7 +3,6 @@
 	"EngineAssociation": "",
 	"Category": "",
 	"Description": "",
-	"Enterprise": true,
 	"Plugins": [
 		{
 			"Name": "MagicLeapMedia",
@@ -15,11 +14,22 @@
 		},
 		{
 			"Name": "OculusVR",
-			"Enabled": false
+			"Enabled": true,
+			"SupportedTargetPlatforms": [
+				"Win32",
+				"Win64",
+				"Android"
+			]
 		},
 		{
 			"Name": "SteamVR",
-			"Enabled": false
+			"Enabled": true,
+			"SupportedTargetPlatforms": [
+				"Win32",
+				"Win64",
+				"Linux",
+				"Mac"
+			]
 		},
 		{
 			"Name": "PythonScriptPlugin",
@@ -64,6 +74,73 @@
 		{
 			"Name": "MediaIOFramework",
 			"Enabled": true
+		},
+		{
+			"Name": "RemoteControl",
+			"Enabled": true
+		},
+		{
+			"Name": "nDisplay",
+			"Enabled": true,
+			"SupportedTargetPlatforms": [
+				"Win64"
+			]
+		},
+		{
+			"Name": "AjaMedia",
+			"Enabled": true,
+			"MarketplaceURL": "com.epicgames.launcher://ue/marketplace/content/91088fb17fb04b6b9eda41aaf8fda21e",
+			"SupportedTargetPlatforms": [
+				"Win64"
+			]
+		},
+		{
+			"Name": "BlackmagicMedia",
+			"Enabled": true,
+			"MarketplaceURL": "com.epicgames.launcher://ue/marketplace/content/433d28bdfb764ec5b2e60c4f45d5cb5c",
+			"SupportedTargetPlatforms": [
+				"Win64"
+			]
+		},
+		{
+			"Name": "AppleProResMedia",
+			"Enabled": true,
+			"SupportedTargetPlatforms": [
+				"Win64"
+			]
+		},
+		{
+			"Name": "HAPMedia",
+			"Enabled": true,
+			"SupportedTargetPlatforms": [
+				"Win64"
+			]
+		},
+		{
+			"Name": "AppleARKit",
+			"Enabled": true,
+			"SupportedTargetPlatforms": [
+				"IOS"
+			]
+		},
+		{
+			"Name": "AppleARKitFaceSupport",
+			"Enabled": true,
+			"SupportedTargetPlatforms": [
+				"IOS"
+			]
+		},
+		{
+			"Name": "PixelStreaming",
+			"Enabled": true
+		},
+			{
+			"Name": "TimecodeSynchronizer",
+			"Enabled": true
+		},
+		{
+			"Name": "VariantManager",
+			"Enabled": true
 		}
 	]
 }
\ No newline at end of file
diff --git a/Templates/TP_ME_VProdBP/Config/DefaultEditor.ini b/Templates/TP_ME_VProdBP/Config/DefaultEditor.ini
new file mode 100644
--- /dev/null
+++ b/Templates/TP_ME_VProdBP/Config/DefaultEditor.ini
@@ -0,0 +1,31 @@
+[/Script/UnrealEd.LevelEditorViewportSettings]
+bInvertMiddleMousePan=True
+bTransparentBoxSelection=True
+bUseDistanceScaledCameraSpeed=True
+
+[ContentBrowser]
+ContentBrowserTab1.SourcesExpanded=True
+
+;Enabling VR Scouting Widgets 
+[/Script/VREditor.VRModeSettings]
+bEnableAutoVREditMode=False
+bAutokeySequences=True
+InteractorHand=Right
+bShowWorldMovementGrid=False
+bShowWorldMovementPostProcess=False
+bShowWorldScaleProgressBar=True
+UIBrightness=1.500000
+GizmoScale=0.800000
+DoubleClickTime=0.250000
+TriggerPressedThreshold_Vive=0.330000
+TriggerPressedThreshold_Rift=0.500000
+InteractorClass=/VirtualProductionUtilities/VR/VirtualScoutingInteractor.VirtualScoutingInteractor_C
+TeleporterClass=BlueprintGeneratedClass'/VirtualProductionUtilities/VR/VirtualScoutingTeleporter.VirtualScoutingTeleporter_C'
+bScaleWorldWithDynamicPivot=True
+bAllowSimultaneousWorldScalingAndRotation=True
+
+[/Script/UnrealEd.EditorPerformanceSettings]
+;bShowFrameRateAndMemory=True
+bThrottleCPUWhenNotForeground=False
+
+
diff --git a/Templates/TP_ME_VProdBP/Config/DefaultEngine.ini b/Templates/TP_ME_VProdBP/Config/DefaultEngine.ini
new file mode 100644
--- /dev/null
+++ b/Templates/TP_ME_VProdBP/Config/DefaultEngine.ini
@@ -0,0 +1,117 @@
+[/Script/HardwareTargeting.HardwareTargetingSettings]
+; quality settings
+TargetedHardwareClass=Desktop
+AppliedTargetedHardwareClass=Desktop
+DefaultGraphicsPerformance=Maximum
+AppliedDefaultGraphicsPerformance=Maximum
+
+[/Script/WindowsTargetPlatform.WindowsTargetSettings]
+;Ndisplay Requirements
+DefaultGraphicsRHI=DefaultGraphicsRHI_DX11
++TargetedRHIs=PCD3D_SM5
+-TargetedRHIs=PCD3D_SM4
+-TargetedRHIs=GLSL_430
+-TargetedRHIs=GLSL_150
+
+[/Script/Engine.RendererSettings]
+; exposure control and tone mapping settings
+r.Tonemapper.Sharpen=2
+r.TemporalAA.Upsampling=True
+r.DefaultFeature.AutoExposure.ExtendDefaultLuminanceRange=True
+r.DefaultFeature.Bloom=False
+r.DefaultFeature.AutoExposure=False
+r.DefaultFeature.MotionBlur=False
+r.SceneRenderTargetResizeMethod=2
+
+
+;dynamic gi settings
+r.LightPropagationVolume=0
+
+; quality settings
+r.ReflectionCaptureResolution=2048
+r.AllowStaticLighting=True
+r.HighResScreenshotDelay=8
+r.DefaultBackBufferPixelFormat=4
+r.AllowGlobalClipPlane=False
+r.GBufferFormat=3
+
+; shader settings
+r.PostProcessing.PropagateAlpha=1
+r.SupportSkyAtmosphereAffectsHeightFog=True
+r.DefaultFeature.LensFlare=False
+r.ClearCoatNormal=False
+r.NormalMapsForStaticLighting=False
+
+
+
+[/Script/EngineSettings.GameMapsSettings]
+;defines what level to be loaded by default in the editor and at playback time
+GameDefaultMap=/Game/VProdProject/Maps/Ndisplay.Ndisplay
+EditorStartupMap=/Game/VProdProject/Maps/Main.Main
+
+[/Script/UnrealEd.UnrealEdEngine]
+; Remove Engine Template maps
+-TemplateMapInfos=(ThumbnailTexture=Texture2D'/Engine/Maps/Templates/Thumbnails/VR-Basic.VR-Basic',Map="/Engine/Maps/Templates/VR-Basic")
+
+
+;Allows for Hardware Accelerated Video Decoding
+[/Script/WmfMediaFactory.WmfMediaSettings]
+AllowNonStandardCodecs=True
+HardwareAcceleratedVideoDecoding=True
+LowLatency=False
+NativeAudioOut=False
+
+;Adds Virtual Scouting Widget for VR Scouting
+[/Script/VPUtilitiesEditor.VPUtilitiesEditorSettings]
+VirtualScoutingUI=/VirtualProductionUtilities/Editor/VirtualScoutingWidget.VirtualScoutingWidget_C
+
+; Enables WebControl API
+WebControl.EnableServerOnStartup=1
+
+;Adds Remote Session for Vcam
+[RemoteSession]
++Channels=(Name=FRemoteSessionFrameBufferChannel,Mode=Write)
++Channels=(Name=FRemoteSessionInputChannel,Mode=Read)
++Channels=(Name=FRemoteSessionXRTrackingChannel,Mode=Read)
+
+;Setup for Media Profiles
+[/Script/MediaFrameworkUtilities.MediaProfileSettings]
+bApplyInCommandlet=False
+MediaSourceProxy=/Game/VprodProject/MediaIO/MediaSource-01.MediaSource-01
+MediaSourceProxy=/Game/VprodProject/MediaIO/MediaSource-02.MediaSource-02
+MediaOutputProxy=/Game/VprodProject/MediaIO/MediaOutput-01.MediaOutput-01
+StartupMediaProfile=None
+
+
+;Setup for Multiuser
+[/Script/Concert.ConcertClientConfig]
+bIsHeadless=False
+bInstallEditorToolbarButton=True
+bAutoConnect=False
+DefaultServerURL=
+DefaultSessionName=
+DefaultSessionToRestore=
+DefaultSaveSessionAs=
+ClientSettings=(DisplayName="",AvatarColor=(R=1.000000,G=1.000000,B=1.000000,A=1.000000),DesktopAvatarActorClass=/ConcertSyncClient/DesktopPresence.DesktopPresence_C,VRAvatarActorClass=/ConcertSyncClient/VRPresence.VRPresence_C,DiscoveryTimeoutSeconds=5,SessionTickFrequencySeconds=1,LatencyCompensationMs=0.000000,Tags=)
+EndpointSettings=(bEnableLogging=False,PurgeProcessedMessageDelaySeconds=30,RemoteEndpointTimeoutSeconds=60)
+
+
+;Setup for Ndisplay Cluster
+[/Script/Engine.Engine]
+GameEngine=/Script/DisplayCluster.DisplayClusterGameEngine
+UnrealEdEngine=/Script/DisplayClusterEditor.DisplayClusterEditorEngine
+GameViewportClientClassName=/Script/DisplayCluster.DisplayClusterViewportClient
+
+;Setup for Ndisplay Cluster
+[/Script/DisplayClusterEditor.DisplayClusterEditorSettings]
+bEnabled=True
+
+[/Script/WindowsTargetPlatform.WindowsTargetSettings]
+;Ndisplay Requirements
+DefaultGraphicsRHI=DefaultGraphicsRHI_DX11
++TargetedRHIs=PCD3D_SM5
+-TargetedRHIs=PCD3D_SM4
+-TargetedRHIs=GLSL_430
+-TargetedRHIs=GLSL_150
+
+
diff --git a/Templates/TP_ME_VProdBP/Config/DefaultGame.ini b/Templates/TP_ME_VProdBP/Config/DefaultGame.ini
new file mode 100644
--- /dev/null
+++ b/Templates/TP_ME_VProdBP/Config/DefaultGame.ini
@@ -0,0 +1,11 @@
+[/Script/EngineSettings.GeneralProjectSettings]
+bUseBorderlessWindow=True
+
+[/Script/MediaFrameworkUtilities.MediaProfileSettings]
+bApplyInCommandlet=False
++MediaSourceProxy=/Game/VprodProject/MediaIO/MediaSource-01.MediaSource-01
++MediaSourceProxy=/Game/VprodProject/MediaIO/MediaSource-02.MediaSource-02
++MediaOutputProxy=/Game/VprodProject/MediaIO/MediaOutput-01.MediaOutput-01
+StartupMediaProfile=None
+
+
diff --git a/Templates/TP_ME_VProdBP/Config/DefaultLightmass.ini b/Templates/TP_ME_VProdBP/Config/DefaultLightmass.ini
new file mode 100644
--- /dev/null
+++ b/Templates/TP_ME_VProdBP/Config/DefaultLightmass.ini
@@ -0,0 +1,22 @@
+; These are tweaked defaults for various lightmass solver and export settings
+; Artist oriented lightmass settings are in the editor UI
+; Documentation for all of these is in UnrealLightmass / Public / SceneExport.h
+; This ini is reloaded every time a lighting build begins, no need to restart
+
+; Warning: overwriting this file with an old version will cause the editor to crash.  This file must be in sync with the editor executable.
+; Instead, create a DefaultLightmass.ini in your project and override just the values you need, then the overrides will continue to work on version upgrades.
+; https://docs.unrealengine.com/latest/INT/Programming/Basics/ConfigurationFiles/
+;
+; For example, in your project's Config/DefaultLightmass.ini:
+; [DevOptions.PrecomputedDynamicObjectLighting]
+; SurfaceLightSampleSpacing=400
+
+[DevOptions.StaticLightingSettings]
+bCompressLightmaps=False
+NumIndirectLightingBounces=5
+NumSkyLightingBounces=5
+IndirectLightingSmoothness=0.9
+IndirectLightingQuality=2.0
+
+[DevOptions.StaticLightingSceneConstants]
+StaticLightingLevelScale=0.5
\ No newline at end of file
diff --git a/Templates/TP_ME_VProdBP/Config/DefaultVirtualProductionUtilities.ini b/Templates/TP_ME_VProdBP/Config/DefaultVirtualProductionUtilities.ini
new file mode 100644
--- /dev/null
+++ b/Templates/TP_ME_VProdBP/Config/DefaultVirtualProductionUtilities.ini
@@ -0,0 +1,13 @@
+
+[/Script/VPUtilitiesEditor.VPUtilitiesEditorSettings]
+VirtualScoutingUI=/VirtualProductionUtilities/Editor/VirtualScoutingWidget.VirtualScoutingWidget_C
+FlightSpeed=0.500000
+GripNavSpeed=0.250000
+bUseMetric=False
+bUseTransformGizmo=False
+bUseGripInertiaDamping=True
+InertiaDamping=0.950000
+bIsHelperSystemEnabled=True
+ScoutingSubsystemEdititorUtilityActorClassPath=/VirtualProductionUtilities/VirtualProductionHelpers.VirtualProductionHelpers_C
+
+
diff --git a/Templates/TP_ME_VProdBP/Config/TemplateDefs.ini b/Templates/TP_ME_VProdBP/Config/TemplateDefs.ini
new file mode 100644
--- /dev/null
+++ b/Templates/TP_ME_VProdBP/Config/TemplateDefs.ini
@@ -0,0 +1,95 @@
+; This file specifies how a template will be treated when creating a project from it.
+; It includes rules for ignoring files, replacing files and changing the contents of files.
+; This file is annotated to help with creating your own templates.
+
+[/Script/GameProjectGeneration.TemplateProjectDefs]
+
+; The categories that this template should appear in. The values in this list refer to the Key value given in TemplateCategories.ini.
+Categories=ME
+
+; The localized display names for this template. Each entry specifies a Language code and Unicode text to display in the New Project dialog.
+LocalizedDisplayNames=(Language="en", Text="Virtual Production")
+
+; The localized descriptions for this template, which is displayed when the template is selected.
+LocalizedDescriptions=(Language="en", Text="This Virtual Production project template contains functionality for VR Scouting, Virutal Camera, SDI Video, Live Compositing, and Ndisplay.")
+; The sort key to use for the template to bypass the default alphabetical sorting.
+SortKey="_2"
+
+; A list of folders to avoid copying when creating a project from this template. 
+; These should contain files only of concern to the template itself and not desired in the user's project.
+FoldersToIgnore=Media
+
+; The files to ignore copying, relative to the template root folder. 
+; %TEMPLATENAME% will be replaced by the template's folder name when resolving these. 
+; Filename comparisons are case-insensitive. 
+; For example, if this template is "TP_Blank" and the user creates "MyProject", then "TP_Blank.uproject" will be copied to "MyProject.uproject".
+FilesToIgnore="%TEMPLATENAME%.uproject"
+FilesToIgnore="%TEMPLATENAME%.png"
+FilesToIgnore="Config/TemplateDefs.ini"
+FilesToIgnore="Manifest.json"
+FilesToIgnore="contents.txt"
+
+; NOTE: For the purposes of the below examples, this template is considered to be named "TP_Blank", and the user has created a project called "MyProject".
+
+; Folders that will be renamed when creating a project from this template.
+; %PROJECTNAME% will be replaced by the project name specified by the user during creation when resolving these rules.
+; For example,  then "Source/TP_Blank" will be copied to "Source/MyProject".
+FolderRenames=(From="Source/%TEMPLATENAME%", To="Source/%PROJECTNAME%")
+FolderRenames=(From="Source/%TEMPLATENAME%Editor", To="Source/%PROJECTNAME%Editor")
+
+; Rules for renaming files when creating a project from this template. 
+; %TEMPLATENAME_UPPERCASE% and %PROJECTNAME_UPPERCASE% are the template and project names in upper case, eg. for "TP_Blank", its upper case representation is "TP_BLANK".
+; %TEMPLATENAME_LOWERCASE% and %PROJECTNAME_LOWERCASE% are the template and project names in lower case, eg. for "TP_Blank", its lower case representation is "tp_blank".
+; A rule does not have to include any of these modifiers. 
+; Extensions is a list of files to consider for these replacements. Eg. ("cpp", "h") will try to apply a rule only to files like "SomeFile.cpp" and "SomeFile.h".
+; bCaseSensitive specifies whether this replacement should be case-sensitive.
+
+; Here are some default rules for renaming files:
+
+; This rule states that a file called "TP_BLANK.cpp" will be renamed to "MYPROJECT.cpp".
+FilenameReplacements=(Extensions=("cpp","h","ini","cs"), From="%TEMPLATENAME_UPPERCASE%", To="%PROJECTNAME_UPPERCASE%", bCaseSensitive=true)
+
+; This rule states that a file called "tp_blank.ini" will be renamed to "myproject.ini"
+FilenameReplacements=(Extensions=("cpp","h","ini","cs"), From="%TEMPLATENAME_LOWERCASE%", To="%PROJECTNAME_LOWERCASE%", bCaseSensitive=true)
+
+; And this rule catches any irregularly-cased files, so that "Tp_bLANK.ini" will be renamed to "MyProject.ini".
+FilenameReplacements=(Extensions=("cpp","h","ini","cs"), From="%TEMPLATENAME%",To="%PROJECTNAME%",bCaseSensitive=false)
+
+; Rules for renaming strings in the contents of files. 
+; The same rules as above apply, but any usages of strings that match these rules inside the file will be changed. This includes all substrings!
+; For example, "TP_Blank" will be replaced by "MyProject", but "TP_Blank_Name" will also be replaced by "MyProject_Name".
+ReplacementsInFiles=(Extensions=("cpp","h","ini","cs"), From="%TEMPLATENAME_UPPERCASE%", To="%PROJECTNAME_UPPERCASE%", bCaseSensitive=true)
+ReplacementsInFiles=(Extensions=("cpp","h","ini","cs"), From="%TEMPLATENAME_LOWERCASE%", To="%PROJECTNAME_LOWERCASE%", bCaseSensitive=true)
+ReplacementsInFiles=(Extensions=("cpp","h","ini","cs"), From="%TEMPLATENAME%", To="%PROJECTNAME%", bCaseSensitive=false)
+
+; The following settings are not enabled for this project, but are included for completeness:
+
+; Whether or not this template can be used to create a project. Default: true
+; If this is set to false, then the template becomes a Feature Pack that can only be referenced from other templates, or added to a project from the Content Browser.
+;bAllowProjectCreation=false
+
+; Which settings on the project settings page should be hidden for this template. Default: none
+; Allowed values are:
+; All - Hides the settings page entirely
+; Languages - Choose between C++ and Blueprint versions.
+; HardwareTarget - Choose between Mobile and Desktop hardware targets.
+; GraphicsPreset - Choose between Maximum and Scalable graphics.
+; StarterContent - Choose whether to add starter content to the project.
+; XR - Choose whether to enable Virtual Reality support for the project.
+;HiddenSettings=All
+HiddenSettings=Languages
+HiddenSettings=HardwareTarget
+HiddenSettings=GraphicsPreset 
+;HiddenSettings=StarterContent 
+HiddenSettings=XR
+
+; The class types that this project uses.
+; This is only for display in the template description, and does not affect project creation.
+;ClassTypes=GameMode, Character, Actor
+
+; The asset types that this project uses.
+; This is only for display in the template description, and does not affect project creation.
+;AssetTypes=Skeleton, Static Mesh, Texture
+
+; Shared feature packs. The files in these packs listed in these structures marked as 'additionalfiles' will be copied on project generation.
+;SharedContentPacks=(MountName="2DSideScroller",DetailLevels=("Standard"))
diff --git a/Templates/TP_ME_VProdBP/TP_ME_VprodBP.uproject b/Templates/TP_ME_VProdBP/TP_ME_VprodBP.uproject
new file mode 100644
--- /dev/null
+++ b/Templates/TP_ME_VProdBP/TP_ME_VprodBP.uproject
@@ -0,0 +1,146 @@
+{
+	"FileVersion": 3,
+	"EngineAssociation": "",
+	"Category": "",
+	"Description": "",
+	"Plugins": [
+		{
+			"Name": "MagicLeapMedia",
+			"Enabled": false
+		},
+		{
+			"Name": "MagicLeap",
+			"Enabled": false
+		},
+		{
+			"Name": "OculusVR",
+			"Enabled": true,
+			"SupportedTargetPlatforms": [
+				"Win32",
+				"Win64",
+				"Android"
+			]
+		},
+		{
+			"Name": "SteamVR",
+			"Enabled": true,
+			"SupportedTargetPlatforms": [
+				"Win32",
+				"Win64",
+				"Linux",
+				"Mac"
+			]
+		},
+		{
+			"Name": "PythonScriptPlugin",
+			"Enabled": true
+		},
+		{
+			"Name": "HDRIBackdrop",
+			"Enabled": true
+		},
+		{
+			"Name": "SunPosition",
+			"Enabled": true
+		},
+		{
+			"Name": "SequencerScripting",
+			"Enabled": true
+		},
+		{
+			"Name": "Takes",
+			"Enabled": true
+		},
+		{
+			"Name": "VirtualProductionUtilities",
+			"Enabled": true
+		},
+		{
+			"Name": "VirtualCamera",
+			"Enabled": true
+		},
+		{
+			"Name": "LiveLink",
+			"Enabled": true
+		},
+		{
+			"Name": "Composure",
+			"Enabled": true
+		},
+		{
+			"Name": "MediaFrameworkUtilities",
+			"Enabled": true
+		},
+		{
+			"Name": "MediaIOFramework",
+			"Enabled": true
+		},
+		{
+			"Name": "RemoteControl",
+			"Enabled": true
+		},
+		{
+			"Name": "nDisplay",
+			"Enabled": true,
+			"SupportedTargetPlatforms": [
+				"Win64"
+			]
+		},
+		{
+			"Name": "AjaMedia",
+			"Enabled": true,
+			"MarketplaceURL": "com.epicgames.launcher://ue/marketplace/content/91088fb17fb04b6b9eda41aaf8fda21e",
+			"SupportedTargetPlatforms": [
+				"Win64"
+			]
+		},
+		{
+			"Name": "BlackmagicMedia",
+			"Enabled": true,
+			"MarketplaceURL": "com.epicgames.launcher://ue/marketplace/content/433d28bdfb764ec5b2e60c4f45d5cb5c",
+			"SupportedTargetPlatforms": [
+				"Win64"
+			]
+		},
+		{
+			"Name": "AppleProResMedia",
+			"Enabled": true,
+			"SupportedTargetPlatforms": [
+				"Win64"
+			]
+		},
+		{
+			"Name": "HAPMedia",
+			"Enabled": true,
+			"SupportedTargetPlatforms": [
+				"Win64"
+			]
+		},
+		{
+			"Name": "AppleARKit",
+			"Enabled": true,
+			"SupportedTargetPlatforms": [
+				"IOS"
+			]
+		},
+		{
+			"Name": "AppleARKitFaceSupport",
+			"Enabled": true,
+			"SupportedTargetPlatforms": [
+				"IOS"
+			]
+		},
+		{
+			"Name": "PixelStreaming",
+			"Enabled": true
+		},
+		{
+			"Name": "TimecodeSynchronizer",
+			"Enabled": true
+		},
+		{
+			"Name": "VariantManager",
+			"Enabled": true
+		}
+	]
+}
\ No newline at end of file
diff --git a/Templates/TP_MFG_BlankBP/Config/DefaultEditor.ini b/Templates/TP_MFG_BlankBP/Config/DefaultEditor.ini
deleted file mode 100644
diff --git a/Templates/TP_MFG_BlankBP/Config/DefaultEngine.ini b/Templates/TP_MFG_BlankBP/Config/DefaultEngine.ini
deleted file mode 100644
diff --git a/Templates/TP_MFG_BlankBP/Config/DefaultGame.ini b/Templates/TP_MFG_BlankBP/Config/DefaultGame.ini
deleted file mode 100644
diff --git a/Templates/TP_MFG_BlankBP/Config/DefaultInput.ini b/Templates/TP_MFG_BlankBP/Config/DefaultInput.ini
deleted file mode 100644
diff --git a/Templates/TP_MFG_BlankBP/Config/DefaultLightmass.ini b/Templates/TP_MFG_BlankBP/Config/DefaultLightmass.ini
deleted file mode 100644
diff --git a/Templates/TP_MFG_BlankBP/Config/TemplateDefs.ini b/Templates/TP_MFG_BlankBP/Config/TemplateDefs.ini
deleted file mode 100644
diff --git a/Templates/TP_MFG_BlankBP/TP_MFG_BlankBP.uproject b/Templates/TP_MFG_BlankBP/TP_MFG_BlankBP.uproject
deleted file mode 100644
diff --git a/Templates/TP_PhotoStudioBP/Config/DefaultEngine.ini b/Templates/TP_PhotoStudioBP/Config/DefaultEngine.ini
--- a/Templates/TP_PhotoStudioBP/Config/DefaultEngine.ini
+++ b/Templates/TP_PhotoStudioBP/Config/DefaultEngine.ini
@@ -17,32 +17,25 @@ r.LightPropagationVolume=0
 ; quality settings
 r.ReflectionCaptureResolution=2048
 r.AllowStaticLighting=True
-r.HighResScreenshotDelay 8
+r.HighResScreenshotDelay=8
 r.DefaultBackBufferPixelFormat=4
 r.AllowGlobalClipPlane=False
 
 ; shader settings
 r.SupportSkyAtmosphereAffectsHeightFog=True
 r.DefaultFeature.LensFlare=True
-r.ClearCoatNormal=False
+r.ClearCoatNormal=True
 r.NormalMapsForStaticLighting=False
-
+r.GenerateMeshDistanceFields=True
 
 [/Script/EngineSettings.GameMapsSettings]
 ;defines what level to be loaded by default in the editor and at playback time
-GameDefaultMap=/Game/Main.Main
-EditorStartupMap=/Game/Main.Main
-
+GameDefaultMap=/Game/PhotoStudio/Levels/Cinematic/Cinematic.Cinematic
+EditorStartupMap=/Game/PhotoStudio/Levels/Cinematic/Cinematic.Cinematic
+GlobalDefaultGameMode=/Game/PhotoStudio/Blueprints/BP_PhotoStageGameMode.BP_PhotoStageGameMode_C
 
 [/Script/UnrealEd.UnrealEdEngine]
 ; Remove Engine Template maps
 -TemplateMapInfos=(ThumbnailTexture=Texture2D'/Engine/Maps/Templates/Thumbnails/Default.Default',Map="/Engine/Maps/Templates/Template_Default")
++TemplateMapInfos=(ThumbnailTexture=Texture2D'/Engine/Maps/Templates/Thumbnails/TimeOfDay.TimeOfDay',Map="/Engine/Maps/Templates/TimeOfDay_Default")
 -TemplateMapInfos=(ThumbnailTexture=Texture2D'/Engine/Maps/Templates/Thumbnails/VR-Basic.VR-Basic',Map="/Engine/Maps/Templates/VR-Basic")
-
-; Add Predefined Template Maps that are suitable for this template
-;+TemplateMapInfos=(ThumbnailTexture=Texture2D'/Game/Developers/Maps/Thumbnails/AssetName.AssetName',Map="/Game/Maps/Templates/MapName")
-+TemplateMapInfos=(ThumbnailTexture=Texture2D'/Game/Developers/LevelTemplates/Thumbnails/Default.Default',Map="/Game/Developers/LevelTemplates/Default")
-+TemplateMapInfos=(ThumbnailTexture=Texture2D'/Game/Developers/LevelTemplates/Thumbnails/Default.Default',Map="/Game/Developers/LevelTemplates/WhiteBackGround")
-+TemplateMapInfos=(ThumbnailTexture=Texture2D'/Game/Developers/LevelTemplates/Thumbnails/Default.Default',Map="/Game/Developers/LevelTemplates/DarkBackGround")
-+TemplateMapInfos=(ThumbnailTexture=Texture2D'/Game/Developers/LevelTemplates/Thumbnails/Default.Default',Map="/Game/Developers/LevelTemplates/Exterior")
-+TemplateMapInfos=(ThumbnailTexture=Texture2D'/Game/Developers/LevelTemplates/Thumbnails/Default.Default',Map="/Game/Developers/LevelTemplates/TurnTable")
diff --git a/Templates/TP_PhotoStudioBP/Config/TemplateDefs.ini b/Templates/TP_PhotoStudioBP/Config/TemplateDefs.ini
--- a/Templates/TP_PhotoStudioBP/Config/TemplateDefs.ini
+++ b/Templates/TP_PhotoStudioBP/Config/TemplateDefs.ini
@@ -11,9 +11,9 @@ Categories=MFG
 LocalizedDisplayNames=(Language="en", Text="Photo Studio")
 
 ; The localized descriptions for this template, which is displayed when the template is selected.
-LocalizedDescriptions=(Language="en", Text="Photo Studio template to create presentations and cinematics.")
+LocalizedDescriptions=(Language="en", Text="This template enables you to create product presentations and cinematics in a photographic studio environment.\n\nThis template enables you to use Datasmith to import models from various design applications and easily light them.")
 ; The sort key to use for the template to bypass the default alphabetical sorting.
-; SortKey="_1"
+SortKey="1"
 
 ; A list of folders to avoid copying when creating a project from this template. 
 ; These should contain files only of concern to the template itself and not desired in the user's project.
@@ -75,7 +75,7 @@ ReplacementsInFiles=(Extensions=("cpp","h","ini","cs"), From="%TEMPLATENAME%", T
 ; HardwareTarget - Choose between Mobile and Desktop hardware targets.
 ; GraphicsPreset - Choose between Maximum and Scalable graphics.
 ; StarterContent - Choose whether to add starter content to the project.
-; VR - Choose whether to enable Virtual Reality support for the project.
+; XR - Choose whether to enable Virtual Reality support for the project.
 ;HiddenSettings=All
 HiddenSettings=Languages
 HiddenSettings=HardwareTarget
diff --git a/Templates/TP_PhotoStudioBP/TP_PhotoStudioBP.uproject b/Templates/TP_PhotoStudioBP/TP_PhotoStudioBP.uproject
--- a/Templates/TP_PhotoStudioBP/TP_PhotoStudioBP.uproject
+++ b/Templates/TP_PhotoStudioBP/TP_PhotoStudioBP.uproject
@@ -17,21 +17,32 @@
 			"Name": "DatasmithImporter",
 			"Enabled": true
 		},
+		{
+			"Name": "DatasmithCADImporter",
+			"Enabled": true,
+			"SupportedTargetPlatforms": [
+				"Win64"
+			]
+		},
 		{
 			"Name": "VariantManager",
-			"Enabled": True
+			"Enabled": true
 		},
 		{
 			"Name": "PythonScriptPlugin",
 			"Enabled": true
 		},
-				{
+		{
 			"Name": "HDRIBackdrop",
 			"Enabled": true
 		},
 		{
 			"Name": "SequencerScripting",
 			"Enabled": true
+		},
+		{
+			"Name": "StaticMeshEditorExtension",
+			"Enabled": true
 		}
 	]
 }
\ No newline at end of file
diff --git a/Templates/TP_Rolling/Source/TP_Rolling/TP_RollingBall.cpp b/Templates/TP_Rolling/Source/TP_Rolling/TP_RollingBall.cpp
--- a/Templates/TP_Rolling/Source/TP_Rolling/TP_RollingBall.cpp
+++ b/Templates/TP_Rolling/Source/TP_Rolling/TP_RollingBall.cpp
@@ -29,8 +29,8 @@ ATP_RollingBall::ATP_RollingBall()
 	SpringArm = CreateDefaultSubobject<USpringArmComponent>(TEXT("SpringArm0"));
 	SpringArm->SetupAttachment(RootComponent);
 	SpringArm->bDoCollisionTest = false;
-	SpringArm->bAbsoluteRotation = true; // Rotation of the ball should not affect rotation of boom
-	SpringArm->RelativeRotation = FRotator(-45.f, 0.f, 0.f);
+	SpringArm->SetUsingAbsoluteRotation(true); // Rotation of the ball should not affect rotation of boom
+	SpringArm->SetRelativeRotation(FRotator(-45.f, 0.f, 0.f));
 	SpringArm->TargetArmLength = 1200.f;
 	SpringArm->bEnableCameraLag = false;
 	SpringArm->CameraLagSpeed = 3.f;
diff --git a/Templates/TP_SideScroller/Source/TP_SideScroller/TP_SideScrollerCharacter.cpp b/Templates/TP_SideScroller/Source/TP_SideScroller/TP_SideScrollerCharacter.cpp
--- a/Templates/TP_SideScroller/Source/TP_SideScroller/TP_SideScrollerCharacter.cpp
+++ b/Templates/TP_SideScroller/Source/TP_SideScroller/TP_SideScrollerCharacter.cpp
@@ -20,11 +20,11 @@ ATP_SideScrollerCharacter::ATP_SideScrollerCharacter()
 	// Create a camera boom attached to the root (capsule)
 	CameraBoom = CreateDefaultSubobject<USpringArmComponent>(TEXT("CameraBoom"));
 	CameraBoom->SetupAttachment(RootComponent);
-	CameraBoom->bAbsoluteRotation = true; // Rotation of the character should not affect rotation of boom
+	CameraBoom->SetUsingAbsoluteRotation(true); // Rotation of the character should not affect rotation of boom
 	CameraBoom->bDoCollisionTest = false;
 	CameraBoom->TargetArmLength = 500.f;
 	CameraBoom->SocketOffset = FVector(0.f,0.f,75.f);
-	CameraBoom->RelativeRotation = FRotator(0.f,180.f,0.f);
+	CameraBoom->SetRelativeRotation(FRotator(0.f,180.f,0.f));
 
 	// Create a camera and attach to boom
 	SideViewCameraComponent = CreateDefaultSubobject<UCameraComponent>(TEXT("SideViewCamera"));
diff --git a/Templates/TP_TopDown/Source/TP_TopDown/TP_TopDownCharacter.cpp b/Templates/TP_TopDown/Source/TP_TopDown/TP_TopDownCharacter.cpp
--- a/Templates/TP_TopDown/Source/TP_TopDown/TP_TopDownCharacter.cpp
+++ b/Templates/TP_TopDown/Source/TP_TopDown/TP_TopDownCharacter.cpp
@@ -31,9 +31,9 @@ ATP_TopDownCharacter::ATP_TopDownCharacter()
 	// Create a camera boom...
 	CameraBoom = CreateDefaultSubobject<USpringArmComponent>(TEXT("CameraBoom"));
 	CameraBoom->SetupAttachment(RootComponent);
-	CameraBoom->bAbsoluteRotation = true; // Don't want arm to rotate when character does
+	CameraBoom->SetUsingAbsoluteRotation(true); // Don't want arm to rotate when character does
 	CameraBoom->TargetArmLength = 800.f;
-	CameraBoom->RelativeRotation = FRotator(-60.f, 0.f, 0.f);
+	CameraBoom->SetRelativeRotation(FRotator(-60.f, 0.f, 0.f));
 	CameraBoom->bDoCollisionTest = false; // Don't want to pull camera in when it collides with level
 
 	// Create a camera...
diff --git a/Templates/TP_TwinStick/Source/TP_TwinStick/TP_TwinStickPawn.cpp b/Templates/TP_TwinStick/Source/TP_TwinStick/TP_TwinStickPawn.cpp
--- a/Templates/TP_TwinStick/Source/TP_TwinStick/TP_TwinStickPawn.cpp
+++ b/Templates/TP_TwinStick/Source/TP_TwinStick/TP_TwinStickPawn.cpp
@@ -34,9 +34,9 @@ ATP_TwinStickPawn::ATP_TwinStickPawn()
 	// Create a camera boom...
 	CameraBoom = CreateDefaultSubobject<USpringArmComponent>(TEXT("CameraBoom"));
 	CameraBoom->SetupAttachment(RootComponent);
-	CameraBoom->bAbsoluteRotation = true; // Don't want arm to rotate when ship does
+	CameraBoom->SetUsingAbsoluteRotation(true); // Don't want arm to rotate when ship does
 	CameraBoom->TargetArmLength = 1200.f;
-	CameraBoom->RelativeRotation = FRotator(-80.f, 0.f, 0.f);
+	CameraBoom->SetRelativeRotation(FRotator(-80.f, 0.f, 0.f));
 	CameraBoom->bDoCollisionTest = false; // Don't want to pull camera in when it collides with level
 
 	// Create a camera...
diff --git a/Templates/TP_Vehicle/Source/TP_Vehicle/TP_VehiclePawn.cpp b/Templates/TP_Vehicle/Source/TP_Vehicle/TP_VehiclePawn.cpp
--- a/Templates/TP_Vehicle/Source/TP_Vehicle/TP_VehiclePawn.cpp
+++ b/Templates/TP_Vehicle/Source/TP_Vehicle/TP_VehiclePawn.cpp
@@ -214,10 +214,10 @@ void ATP_VehiclePawn::Tick(float Delta)
 	{
 		if ( (InputComponent) && (bInCarCameraActive == true ))
 		{
-			FRotator HeadRotation = InternalCamera->RelativeRotation;
+			FRotator HeadRotation = InternalCamera->GetRelativeRotation();
 			HeadRotation.Pitch += InputComponent->GetAxisValue(LookUpBinding);
 			HeadRotation.Yaw += InputComponent->GetAxisValue(LookRightBinding);
-			InternalCamera->RelativeRotation = HeadRotation;
+			InternalCamera->SetRelativeRotation(HeadRotation);
 		}
 	}
 }
diff --git a/Templates/TP_VehicleAdv/Source/TP_VehicleAdv/TP_VehicleAdvPawn.cpp b/Templates/TP_VehicleAdv/Source/TP_VehicleAdv/TP_VehicleAdvPawn.cpp
--- a/Templates/TP_VehicleAdv/Source/TP_VehicleAdv/TP_VehicleAdvPawn.cpp
+++ b/Templates/TP_VehicleAdv/Source/TP_VehicleAdv/TP_VehicleAdvPawn.cpp
@@ -272,10 +272,10 @@ void ATP_VehicleAdvPawn::Tick(float Delta)
 	{
 		if ( (InputComponent) && (bInCarCameraActive == true ))
 		{
-			FRotator HeadRotation = InternalCamera->RelativeRotation;
+			FRotator HeadRotation = InternalCamera->GetRelativeRotation();
 			HeadRotation.Pitch += InputComponent->GetAxisValue(LookUpBinding);
 			HeadRotation.Yaw += InputComponent->GetAxisValue(LookRightBinding);
-			InternalCamera->RelativeRotation = HeadRotation;
+			InternalCamera->SetRelativeRotation(HeadRotation);
 		}
 	}	
 
diff --git a/Templates/TP_VirtualRealityBP/Config/DefaultInput.ini b/Templates/TP_VirtualRealityBP/Config/DefaultInput.ini
--- a/Templates/TP_VirtualRealityBP/Config/DefaultInput.ini
+++ b/Templates/TP_VirtualRealityBP/Config/DefaultInput.ini
@@ -1,38 +1,17 @@
 
 [/Script/Engine.InputSettings]
--AxisConfig=(AxisKeyName="Gamepad_RightY",AxisProperties=(DeadZone=0.250000,Sensitivity=1.000000,Exponent=1.000000,bInvert=False))
--AxisConfig=(AxisKeyName="MouseX",AxisProperties=(DeadZone=0.000000,Sensitivity=0.070000,Exponent=1.000000,bInvert=False))
--AxisConfig=(AxisKeyName="MouseY",AxisProperties=(DeadZone=0.000000,Sensitivity=0.070000,Exponent=1.000000,bInvert=False))
--AxisConfig=(AxisKeyName="MouseWheelAxis",AxisProperties=(DeadZone=0.000000,Sensitivity=1.000000,Exponent=1.000000,bInvert=False))
--AxisConfig=(AxisKeyName="Gamepad_LeftTriggerAxis",AxisProperties=(DeadZone=0.000000,Sensitivity=1.000000,Exponent=1.000000,bInvert=False))
--AxisConfig=(AxisKeyName="Gamepad_RightTriggerAxis",AxisProperties=(DeadZone=0.000000,Sensitivity=1.000000,Exponent=1.000000,bInvert=False))
--AxisConfig=(AxisKeyName="MotionController_Left_Thumbstick_X",AxisProperties=(DeadZone=0.000000,Sensitivity=1.000000,Exponent=1.000000,bInvert=False))
--AxisConfig=(AxisKeyName="MotionController_Left_Thumbstick_Y",AxisProperties=(DeadZone=0.000000,Sensitivity=1.000000,Exponent=1.000000,bInvert=False))
--AxisConfig=(AxisKeyName="MotionController_Left_TriggerAxis",AxisProperties=(DeadZone=0.000000,Sensitivity=1.000000,Exponent=1.000000,bInvert=False))
--AxisConfig=(AxisKeyName="MotionController_Left_Grip1Axis",AxisProperties=(DeadZone=0.000000,Sensitivity=1.000000,Exponent=1.000000,bInvert=False))
--AxisConfig=(AxisKeyName="MotionController_Left_Grip2Axis",AxisProperties=(DeadZone=0.000000,Sensitivity=1.000000,Exponent=1.000000,bInvert=False))
--AxisConfig=(AxisKeyName="MotionController_Right_Thumbstick_X",AxisProperties=(DeadZone=0.000000,Sensitivity=1.000000,Exponent=1.000000,bInvert=False))
--AxisConfig=(AxisKeyName="MotionController_Right_Thumbstick_Y",AxisProperties=(DeadZone=0.000000,Sensitivity=1.000000,Exponent=1.000000,bInvert=False))
--AxisConfig=(AxisKeyName="MotionController_Right_TriggerAxis",AxisProperties=(DeadZone=0.000000,Sensitivity=1.000000,Exponent=1.000000,bInvert=False))
--AxisConfig=(AxisKeyName="MotionController_Right_Grip1Axis",AxisProperties=(DeadZone=0.000000,Sensitivity=1.000000,Exponent=1.000000,bInvert=False))
--AxisConfig=(AxisKeyName="MotionController_Right_Grip2Axis",AxisProperties=(DeadZone=0.000000,Sensitivity=1.000000,Exponent=1.000000,bInvert=False))
--AxisConfig=(AxisKeyName="Gamepad_Special_Left_X",AxisProperties=(DeadZone=0.000000,Sensitivity=1.000000,Exponent=1.000000,bInvert=False))
--AxisConfig=(AxisKeyName="Gamepad_Special_Left_Y",AxisProperties=(DeadZone=0.000000,Sensitivity=1.000000,Exponent=1.000000,bInvert=False))
--AxisConfig=(AxisKeyName="OculusTouch_Left_Thumbstick",AxisProperties=(DeadZone=0.000000,Sensitivity=1.000000,Exponent=1.000000,bInvert=False))
--AxisConfig=(AxisKeyName="OculusTouch_Left_FaceButton1",AxisProperties=(DeadZone=0.000000,Sensitivity=1.000000,Exponent=1.000000,bInvert=False))
--AxisConfig=(AxisKeyName="OculusTouch_Left_Trigger",AxisProperties=(DeadZone=0.000000,Sensitivity=1.000000,Exponent=1.000000,bInvert=False))
--AxisConfig=(AxisKeyName="OculusTouch_Left_FaceButton2",AxisProperties=(DeadZone=0.000000,Sensitivity=1.000000,Exponent=1.000000,bInvert=False))
--AxisConfig=(AxisKeyName="OculusTouch_Left_IndexPointing",AxisProperties=(DeadZone=0.000000,Sensitivity=1.000000,Exponent=1.000000,bInvert=False))
--AxisConfig=(AxisKeyName="OculusTouch_Left_ThumbUp",AxisProperties=(DeadZone=0.000000,Sensitivity=1.000000,Exponent=1.000000,bInvert=False))
--AxisConfig=(AxisKeyName="OculusTouch_Right_Thumbstick",AxisProperties=(DeadZone=0.000000,Sensitivity=1.000000,Exponent=1.000000,bInvert=False))
--AxisConfig=(AxisKeyName="OculusTouch_Right_FaceButton1",AxisProperties=(DeadZone=0.000000,Sensitivity=1.000000,Exponent=1.000000,bInvert=False))
--AxisConfig=(AxisKeyName="OculusTouch_Right_Trigger",AxisProperties=(DeadZone=0.000000,Sensitivity=1.000000,Exponent=1.000000,bInvert=False))
--AxisConfig=(AxisKeyName="OculusTouch_Right_FaceButton2",AxisProperties=(DeadZone=0.000000,Sensitivity=1.000000,Exponent=1.000000,bInvert=False))
--AxisConfig=(AxisKeyName="OculusTouch_Right_IndexPointing",AxisProperties=(DeadZone=0.000000,Sensitivity=1.000000,Exponent=1.000000,bInvert=False))
--AxisConfig=(AxisKeyName="OculusTouch_Right_ThumbUp",AxisProperties=(DeadZone=0.000000,Sensitivity=1.000000,Exponent=1.000000,bInvert=False))
--AxisConfig=(AxisKeyName="Gamepad_LeftX",AxisProperties=(DeadZone=0.250000,Sensitivity=1.000000,Exponent=1.000000,bInvert=False))
--AxisConfig=(AxisKeyName="Gamepad_LeftY",AxisProperties=(DeadZone=0.250000,Sensitivity=1.000000,Exponent=1.000000,bInvert=False))
--AxisConfig=(AxisKeyName="Gamepad_RightX",AxisProperties=(DeadZone=0.250000,Sensitivity=1.000000,Exponent=1.000000,bInvert=False))
+-AxisConfig=(AxisKeyName="Gamepad_LeftX",AxisProperties=(DeadZone=0.25,Exponent=1.f,Sensitivity=1.f))
+-AxisConfig=(AxisKeyName="Gamepad_LeftY",AxisProperties=(DeadZone=0.25,Exponent=1.f,Sensitivity=1.f))
+-AxisConfig=(AxisKeyName="Gamepad_RightX",AxisProperties=(DeadZone=0.25,Exponent=1.f,Sensitivity=1.f))
+-AxisConfig=(AxisKeyName="Gamepad_RightY",AxisProperties=(DeadZone=0.25,Exponent=1.f,Sensitivity=1.f))
+-AxisConfig=(AxisKeyName="MouseX",AxisProperties=(DeadZone=0.f,Exponent=1.f,Sensitivity=0.07f))
+-AxisConfig=(AxisKeyName="MouseY",AxisProperties=(DeadZone=0.f,Exponent=1.f,Sensitivity=0.07f))
++AxisConfig=(AxisKeyName="OculusTouch_Right_FaceButton2",AxisProperties=(DeadZone=0.000000,Sensitivity=1.000000,Exponent=1.000000,bInvert=False))
++AxisConfig=(AxisKeyName="OculusTouch_Right_IndexPointing",AxisProperties=(DeadZone=0.000000,Sensitivity=1.000000,Exponent=1.000000,bInvert=False))
++AxisConfig=(AxisKeyName="OculusTouch_Right_ThumbUp",AxisProperties=(DeadZone=0.000000,Sensitivity=1.000000,Exponent=1.000000,bInvert=False))
++AxisConfig=(AxisKeyName="Gamepad_LeftX",AxisProperties=(DeadZone=0.250000,Sensitivity=1.000000,Exponent=1.000000,bInvert=False))
++AxisConfig=(AxisKeyName="Gamepad_LeftY",AxisProperties=(DeadZone=0.250000,Sensitivity=1.000000,Exponent=1.000000,bInvert=False))
++AxisConfig=(AxisKeyName="Gamepad_RightX",AxisProperties=(DeadZone=0.250000,Sensitivity=1.000000,Exponent=1.000000,bInvert=False))
 +AxisConfig=(AxisKeyName="Gamepad_RightY",AxisProperties=(DeadZone=0.250000,Sensitivity=1.000000,Exponent=1.000000,bInvert=False))
 +AxisConfig=(AxisKeyName="MouseX",AxisProperties=(DeadZone=0.000000,Sensitivity=0.070000,Exponent=1.000000,bInvert=False))
 +AxisConfig=(AxisKeyName="MouseY",AxisProperties=(DeadZone=0.000000,Sensitivity=0.070000,Exponent=1.000000,bInvert=False))
@@ -60,52 +39,126 @@
 +AxisConfig=(AxisKeyName="OculusTouch_Right_Thumbstick",AxisProperties=(DeadZone=0.000000,Sensitivity=1.000000,Exponent=1.000000,bInvert=False))
 +AxisConfig=(AxisKeyName="OculusTouch_Right_FaceButton1",AxisProperties=(DeadZone=0.000000,Sensitivity=1.000000,Exponent=1.000000,bInvert=False))
 +AxisConfig=(AxisKeyName="OculusTouch_Right_Trigger",AxisProperties=(DeadZone=0.000000,Sensitivity=1.000000,Exponent=1.000000,bInvert=False))
-+AxisConfig=(AxisKeyName="OculusTouch_Right_FaceButton2",AxisProperties=(DeadZone=0.000000,Sensitivity=1.000000,Exponent=1.000000,bInvert=False))
-+AxisConfig=(AxisKeyName="OculusTouch_Right_IndexPointing",AxisProperties=(DeadZone=0.000000,Sensitivity=1.000000,Exponent=1.000000,bInvert=False))
-+AxisConfig=(AxisKeyName="OculusTouch_Right_ThumbUp",AxisProperties=(DeadZone=0.000000,Sensitivity=1.000000,Exponent=1.000000,bInvert=False))
-+AxisConfig=(AxisKeyName="Gamepad_LeftX",AxisProperties=(DeadZone=0.250000,Sensitivity=1.000000,Exponent=1.000000,bInvert=False))
-+AxisConfig=(AxisKeyName="Gamepad_LeftY",AxisProperties=(DeadZone=0.250000,Sensitivity=1.000000,Exponent=1.000000,bInvert=False))
-+AxisConfig=(AxisKeyName="Gamepad_RightX",AxisProperties=(DeadZone=0.250000,Sensitivity=1.000000,Exponent=1.000000,bInvert=False))
++AxisConfig=(AxisKeyName="Daydream_Left_Trackpad_X",AxisProperties=(DeadZone=0.000000,Sensitivity=1.000000,Exponent=1.000000,bInvert=False))
++AxisConfig=(AxisKeyName="Daydream_Left_Trackpad_Y",AxisProperties=(DeadZone=0.000000,Sensitivity=1.000000,Exponent=1.000000,bInvert=False))
++AxisConfig=(AxisKeyName="Daydream_Right_Trackpad_X",AxisProperties=(DeadZone=0.000000,Sensitivity=1.000000,Exponent=1.000000,bInvert=False))
++AxisConfig=(AxisKeyName="Daydream_Right_Trackpad_Y",AxisProperties=(DeadZone=0.000000,Sensitivity=1.000000,Exponent=1.000000,bInvert=False))
++AxisConfig=(AxisKeyName="Vive_Left_Trigger_Axis",AxisProperties=(DeadZone=0.000000,Sensitivity=1.000000,Exponent=1.000000,bInvert=False))
++AxisConfig=(AxisKeyName="Vive_Left_Trackpad_X",AxisProperties=(DeadZone=0.000000,Sensitivity=1.000000,Exponent=1.000000,bInvert=False))
++AxisConfig=(AxisKeyName="Vive_Left_Trackpad_Y",AxisProperties=(DeadZone=0.000000,Sensitivity=1.000000,Exponent=1.000000,bInvert=False))
++AxisConfig=(AxisKeyName="Vive_Right_Trigger_Axis",AxisProperties=(DeadZone=0.000000,Sensitivity=1.000000,Exponent=1.000000,bInvert=False))
++AxisConfig=(AxisKeyName="Vive_Right_Trackpad_X",AxisProperties=(DeadZone=0.000000,Sensitivity=1.000000,Exponent=1.000000,bInvert=False))
++AxisConfig=(AxisKeyName="Vive_Right_Trackpad_Y",AxisProperties=(DeadZone=0.000000,Sensitivity=1.000000,Exponent=1.000000,bInvert=False))
++AxisConfig=(AxisKeyName="MixedReality_Left_Trigger_Axis",AxisProperties=(DeadZone=0.000000,Sensitivity=1.000000,Exponent=1.000000,bInvert=False))
++AxisConfig=(AxisKeyName="MixedReality_Left_Thumbstick_X",AxisProperties=(DeadZone=0.000000,Sensitivity=1.000000,Exponent=1.000000,bInvert=False))
++AxisConfig=(AxisKeyName="MixedReality_Left_Thumbstick_Y",AxisProperties=(DeadZone=0.000000,Sensitivity=1.000000,Exponent=1.000000,bInvert=False))
++AxisConfig=(AxisKeyName="MixedReality_Left_Trackpad_X",AxisProperties=(DeadZone=0.000000,Sensitivity=1.000000,Exponent=1.000000,bInvert=False))
++AxisConfig=(AxisKeyName="MixedReality_Left_Trackpad_Y",AxisProperties=(DeadZone=0.000000,Sensitivity=1.000000,Exponent=1.000000,bInvert=False))
++AxisConfig=(AxisKeyName="MixedReality_Right_Trigger_Axis",AxisProperties=(DeadZone=0.000000,Sensitivity=1.000000,Exponent=1.000000,bInvert=False))
++AxisConfig=(AxisKeyName="MixedReality_Right_Thumbstick_X",AxisProperties=(DeadZone=0.000000,Sensitivity=1.000000,Exponent=1.000000,bInvert=False))
++AxisConfig=(AxisKeyName="MixedReality_Right_Thumbstick_Y",AxisProperties=(DeadZone=0.000000,Sensitivity=1.000000,Exponent=1.000000,bInvert=False))
++AxisConfig=(AxisKeyName="MixedReality_Right_Trackpad_X",AxisProperties=(DeadZone=0.000000,Sensitivity=1.000000,Exponent=1.000000,bInvert=False))
++AxisConfig=(AxisKeyName="MixedReality_Right_Trackpad_Y",AxisProperties=(DeadZone=0.000000,Sensitivity=1.000000,Exponent=1.000000,bInvert=False))
++AxisConfig=(AxisKeyName="OculusGo_Left_Trackpad_X",AxisProperties=(DeadZone=0.000000,Sensitivity=1.000000,Exponent=1.000000,bInvert=False))
++AxisConfig=(AxisKeyName="OculusGo_Left_Trackpad_Y",AxisProperties=(DeadZone=0.000000,Sensitivity=1.000000,Exponent=1.000000,bInvert=False))
++AxisConfig=(AxisKeyName="OculusGo_Right_Trackpad_X",AxisProperties=(DeadZone=0.000000,Sensitivity=1.000000,Exponent=1.000000,bInvert=False))
++AxisConfig=(AxisKeyName="OculusGo_Right_Trackpad_Y",AxisProperties=(DeadZone=0.000000,Sensitivity=1.000000,Exponent=1.000000,bInvert=False))
++AxisConfig=(AxisKeyName="OculusTouch_Left_Grip_Axis",AxisProperties=(DeadZone=0.000000,Sensitivity=1.000000,Exponent=1.000000,bInvert=False))
++AxisConfig=(AxisKeyName="OculusTouch_Left_Trigger_Axis",AxisProperties=(DeadZone=0.000000,Sensitivity=1.000000,Exponent=1.000000,bInvert=False))
++AxisConfig=(AxisKeyName="OculusTouch_Left_Thumbstick_X",AxisProperties=(DeadZone=0.000000,Sensitivity=1.000000,Exponent=1.000000,bInvert=False))
++AxisConfig=(AxisKeyName="OculusTouch_Left_Thumbstick_Y",AxisProperties=(DeadZone=0.000000,Sensitivity=1.000000,Exponent=1.000000,bInvert=False))
++AxisConfig=(AxisKeyName="OculusTouch_Right_Grip_Axis",AxisProperties=(DeadZone=0.000000,Sensitivity=1.000000,Exponent=1.000000,bInvert=False))
++AxisConfig=(AxisKeyName="OculusTouch_Right_Trigger_Axis",AxisProperties=(DeadZone=0.000000,Sensitivity=1.000000,Exponent=1.000000,bInvert=False))
++AxisConfig=(AxisKeyName="OculusTouch_Right_Thumbstick_X",AxisProperties=(DeadZone=0.000000,Sensitivity=1.000000,Exponent=1.000000,bInvert=False))
++AxisConfig=(AxisKeyName="OculusTouch_Right_Thumbstick_Y",AxisProperties=(DeadZone=0.000000,Sensitivity=1.000000,Exponent=1.000000,bInvert=False))
++AxisConfig=(AxisKeyName="ValveIndex_Left_Grip_Axis",AxisProperties=(DeadZone=0.000000,Sensitivity=1.000000,Exponent=1.000000,bInvert=False))
++AxisConfig=(AxisKeyName="ValveIndex_Left_Grip_Force",AxisProperties=(DeadZone=0.000000,Sensitivity=1.000000,Exponent=1.000000,bInvert=False))
++AxisConfig=(AxisKeyName="ValveIndex_Left_Trigger_Axis",AxisProperties=(DeadZone=0.000000,Sensitivity=1.000000,Exponent=1.000000,bInvert=False))
++AxisConfig=(AxisKeyName="ValveIndex_Left_Thumbstick_X",AxisProperties=(DeadZone=0.000000,Sensitivity=1.000000,Exponent=1.000000,bInvert=False))
++AxisConfig=(AxisKeyName="ValveIndex_Left_Thumbstick_Y",AxisProperties=(DeadZone=0.000000,Sensitivity=1.000000,Exponent=1.000000,bInvert=False))
++AxisConfig=(AxisKeyName="ValveIndex_Left_Trackpad_X",AxisProperties=(DeadZone=0.000000,Sensitivity=1.000000,Exponent=1.000000,bInvert=False))
++AxisConfig=(AxisKeyName="ValveIndex_Left_Trackpad_Y",AxisProperties=(DeadZone=0.000000,Sensitivity=1.000000,Exponent=1.000000,bInvert=False))
++AxisConfig=(AxisKeyName="ValveIndex_Left_Trackpad_Force",AxisProperties=(DeadZone=0.000000,Sensitivity=1.000000,Exponent=1.000000,bInvert=False))
++AxisConfig=(AxisKeyName="ValveIndex_Left_Trackpad_Touch",AxisProperties=(DeadZone=0.000000,Sensitivity=1.000000,Exponent=1.000000,bInvert=False))
++AxisConfig=(AxisKeyName="ValveIndex_Right_Grip_Axis",AxisProperties=(DeadZone=0.000000,Sensitivity=1.000000,Exponent=1.000000,bInvert=False))
++AxisConfig=(AxisKeyName="ValveIndex_Right_Grip_Force",AxisProperties=(DeadZone=0.000000,Sensitivity=1.000000,Exponent=1.000000,bInvert=False))
++AxisConfig=(AxisKeyName="ValveIndex_Right_Trigger_Axis",AxisProperties=(DeadZone=0.000000,Sensitivity=1.000000,Exponent=1.000000,bInvert=False))
++AxisConfig=(AxisKeyName="ValveIndex_Right_Thumbstick_X",AxisProperties=(DeadZone=0.000000,Sensitivity=1.000000,Exponent=1.000000,bInvert=False))
++AxisConfig=(AxisKeyName="ValveIndex_Right_Thumbstick_Y",AxisProperties=(DeadZone=0.000000,Sensitivity=1.000000,Exponent=1.000000,bInvert=False))
++AxisConfig=(AxisKeyName="ValveIndex_Right_Trackpad_X",AxisProperties=(DeadZone=0.000000,Sensitivity=1.000000,Exponent=1.000000,bInvert=False))
++AxisConfig=(AxisKeyName="ValveIndex_Right_Trackpad_Y",AxisProperties=(DeadZone=0.000000,Sensitivity=1.000000,Exponent=1.000000,bInvert=False))
++AxisConfig=(AxisKeyName="ValveIndex_Right_Trackpad_Force",AxisProperties=(DeadZone=0.000000,Sensitivity=1.000000,Exponent=1.000000,bInvert=False))
 bAltEnterTogglesFullscreen=True
+bF11TogglesFullscreen=True
 bUseMouseForTouch=False
 bEnableMouseSmoothing=True
 bEnableFOVScaling=True
-FOVScale=0.011110
-DoubleClickTime=0.200000
 bCaptureMouseOnLaunch=True
-DefaultViewportMouseCaptureMode=CapturePermanently_IncludingInitialMouseDown
 bDefaultViewportMouseLock=True
-DefaultViewportMouseLockMode=LockOnCapture
--ActionMappings=(ActionName="TeleportRight",Key=MotionController_Right_Thumbstick,bShift=False,bCtrl=False,bAlt=False,bCmd=False)
--ActionMappings=(ActionName="GrabLeft",Key=MotionController_Left_Trigger,bShift=False,bCtrl=False,bAlt=False,bCmd=False)
--ActionMappings=(ActionName="GrabRight",Key=MotionController_Right_Trigger,bShift=False,bCtrl=False,bAlt=False,bCmd=False)
--ActionMappings=(ActionName="TeleportLeft",Key=MotionController_Left_Thumbstick,bShift=False,bCtrl=False,bAlt=False,bCmd=False)
--ActionMappings=(ActionName="HMDTeleport",Key=Gamepad_FaceButton_Bottom,bShift=False,bCtrl=False,bAlt=False,bCmd=False)
--ActionMappings=(ActionName="HMDTeleport",Key=SpaceBar,bShift=False,bCtrl=False,bAlt=False,bCmd=False)
--ActionMappings=(ActionName="TeleportLeft",Key=MotionController_Left_FaceButton1,bShift=False,bCtrl=False,bAlt=False,bCmd=False)
--ActionMappings=(ActionName="TeleportRight",Key=None,bShift=False,bCtrl=False,bAlt=False,bCmd=False)
-+ActionMappings=(ActionName="TeleportRight",Key=MotionController_Right_Thumbstick,bShift=False,bCtrl=False,bAlt=False,bCmd=False)
-+ActionMappings=(ActionName="GrabLeft",Key=MotionController_Left_Trigger,bShift=False,bCtrl=False,bAlt=False,bCmd=False)
-+ActionMappings=(ActionName="GrabRight",Key=MotionController_Right_Trigger,bShift=False,bCtrl=False,bAlt=False,bCmd=False)
-+ActionMappings=(ActionName="TeleportLeft",Key=MotionController_Left_Thumbstick,bShift=False,bCtrl=False,bAlt=False,bCmd=False)
-+ActionMappings=(ActionName="HMDTeleport",Key=Gamepad_FaceButton_Bottom,bShift=False,bCtrl=False,bAlt=False,bCmd=False)
-+ActionMappings=(ActionName="HMDTeleport",Key=SpaceBar,bShift=False,bCtrl=False,bAlt=False,bCmd=False)
-+ActionMappings=(ActionName="TeleportLeft",Key=MotionController_Left_FaceButton1,bShift=False,bCtrl=False,bAlt=False,bCmd=False)
-+ActionMappings=(ActionName="TeleportRight",Key=MotionController_Right_FaceButton1,bShift=False,bCtrl=False,bAlt=False,bCmd=False)
--AxisMappings=(AxisName="TeleportDirectionUp",Key=Gamepad_LeftY,Scale=1.000000)
--AxisMappings=(AxisName="TeleportDirectionRight",Key=Gamepad_LeftX,Scale=1.000000)
--AxisMappings=(AxisName="MotionControllerThumbLeft_Y",Key=MotionController_Left_Thumbstick_Y,Scale=-1.000000)
--AxisMappings=(AxisName="MotionControllerThumbLeft_X",Key=MotionController_Left_Thumbstick_X,Scale=1.000000)
--AxisMappings=(AxisName="MotionControllerThumbRight_Y",Key=MotionController_Right_Thumbstick_Y,Scale=-1.000000)
--AxisMappings=(AxisName="MotionControllerThumbRight_X",Key=MotionController_Right_Thumbstick_X,Scale=1.000000)
-+AxisMappings=(AxisName="TeleportDirectionUp",Key=Gamepad_LeftY,Scale=1.000000)
-+AxisMappings=(AxisName="TeleportDirectionRight",Key=Gamepad_LeftX,Scale=1.000000)
-+AxisMappings=(AxisName="MotionControllerThumbLeft_Y",Key=MotionController_Left_Thumbstick_Y,Scale=-1.000000)
-+AxisMappings=(AxisName="MotionControllerThumbLeft_X",Key=MotionController_Left_Thumbstick_X,Scale=1.000000)
-+AxisMappings=(AxisName="MotionControllerThumbRight_Y",Key=MotionController_Right_Thumbstick_Y,Scale=-1.000000)
-+AxisMappings=(AxisName="MotionControllerThumbRight_X",Key=MotionController_Right_Thumbstick_X,Scale=1.000000)
 bAlwaysShowTouchInterface=False
 bShowConsoleOnFourFingerTap=True
+bEnableGestureRecognizer=False
+bUseAutocorrect=False
+DefaultViewportMouseCaptureMode=CapturePermanently_IncludingInitialMouseDown
+DefaultViewportMouseLockMode=LockOnCapture
+FOVScale=0.011110
+DoubleClickTime=0.200000
++ActionMappings=(ActionName="HMDTeleport",bShift=False,bCtrl=False,bAlt=False,bCmd=False,Key=Gamepad_FaceButton_Bottom)
++ActionMappings=(ActionName="HMDTeleport",bShift=False,bCtrl=False,bAlt=False,bCmd=False,Key=SpaceBar)
++ActionMappings=(ActionName="TeleportRight",bShift=False,bCtrl=False,bAlt=False,bCmd=False,Key=Vive_Right_Trackpad_Click)
++ActionMappings=(ActionName="TeleportRight",bShift=False,bCtrl=False,bAlt=False,bCmd=False,Key=OculusTouch_Right_Thumbstick_Click)
++ActionMappings=(ActionName="TeleportRight",bShift=False,bCtrl=False,bAlt=False,bCmd=False,Key=ValveIndex_Right_Thumbstick_Click)
++ActionMappings=(ActionName="TeleportRight",bShift=False,bCtrl=False,bAlt=False,bCmd=False,Key=MixedReality_Right_Thumbstick_Click)
++ActionMappings=(ActionName="TeleportRight",bShift=False,bCtrl=False,bAlt=False,bCmd=False,Key=OculusGo_Right_Trackpad_Click)
++ActionMappings=(ActionName="TeleportRight",bShift=False,bCtrl=False,bAlt=False,bCmd=False,Key=Daydream_Right_Trackpad_Click)
++ActionMappings=(ActionName="TeleportRight",bShift=False,bCtrl=False,bAlt=False,bCmd=False,Key=OculusTouch_Right_A_Click)
++ActionMappings=(ActionName="TeleportLeft",bShift=False,bCtrl=False,bAlt=False,bCmd=False,Key=Vive_Left_Trackpad_Click)
++ActionMappings=(ActionName="TeleportLeft",bShift=False,bCtrl=False,bAlt=False,bCmd=False,Key=OculusTouch_Left_Thumbstick_Click)
++ActionMappings=(ActionName="TeleportLeft",bShift=False,bCtrl=False,bAlt=False,bCmd=False,Key=ValveIndex_Left_Thumbstick_Click)
++ActionMappings=(ActionName="TeleportLeft",bShift=False,bCtrl=False,bAlt=False,bCmd=False,Key=MixedReality_Left_Thumbstick_Click)
++ActionMappings=(ActionName="TeleportLeft",bShift=False,bCtrl=False,bAlt=False,bCmd=False,Key=OculusGo_Left_Trackpad_Click)
++ActionMappings=(ActionName="TeleportLeft",bShift=False,bCtrl=False,bAlt=False,bCmd=False,Key=Daydream_Left_Trackpad_Click)
++ActionMappings=(ActionName="TeleportLeft",bShift=False,bCtrl=False,bAlt=False,bCmd=False,Key=OculusTouch_Left_X_Click)
++ActionMappings=(ActionName="GrabLeft",bShift=False,bCtrl=False,bAlt=False,bCmd=False,Key=Vive_Left_Trigger_Click)
++ActionMappings=(ActionName="GrabLeft",bShift=False,bCtrl=False,bAlt=False,bCmd=False,Key=OculusTouch_Left_Grip_Click)
++ActionMappings=(ActionName="GrabLeft",bShift=False,bCtrl=False,bAlt=False,bCmd=False,Key=ValveIndex_Left_Grip_Click)
++ActionMappings=(ActionName="GrabLeft",bShift=False,bCtrl=False,bAlt=False,bCmd=False,Key=MixedReality_Left_Trigger_Click)
++ActionMappings=(ActionName="GrabLeft",bShift=False,bCtrl=False,bAlt=False,bCmd=False,Key=OculusGo_Left_Trigger_Click)
++ActionMappings=(ActionName="GrabRight",bShift=False,bCtrl=False,bAlt=False,bCmd=False,Key=Vive_Right_Trigger_Click)
++ActionMappings=(ActionName="GrabRight",bShift=False,bCtrl=False,bAlt=False,bCmd=False,Key=OculusTouch_Right_Grip_Click)
++ActionMappings=(ActionName="GrabRight",bShift=False,bCtrl=False,bAlt=False,bCmd=False,Key=ValveIndex_Right_Grip_Click)
++ActionMappings=(ActionName="GrabRight",bShift=False,bCtrl=False,bAlt=False,bCmd=False,Key=MixedReality_Right_Trigger_Click)
++ActionMappings=(ActionName="GrabRight",bShift=False,bCtrl=False,bAlt=False,bCmd=False,Key=OculusGo_Right_Trigger_Click)
++ActionMappings=(ActionName="TeleportRight",bShift=False,bCtrl=False,bAlt=False,bCmd=False,Key=ValveIndex_Right_A_Click)
++ActionMappings=(ActionName="TeleportLeft",bShift=False,bCtrl=False,bAlt=False,bCmd=False,Key=ValveIndex_Left_A_Click)
++ActionMappings=(ActionName="GrabLeft",bShift=False,bCtrl=False,bAlt=False,bCmd=False,Key=Daydream_Left_Select_Click)
++ActionMappings=(ActionName="GrabRight",bShift=False,bCtrl=False,bAlt=False,bCmd=False,Key=Daydream_Right_Select_Click)
++AxisMappings=(AxisName="TeleportDirectionUp",Scale=1.000000,Key=Gamepad_LeftY)
++AxisMappings=(AxisName="TeleportDirectionRight",Scale=1.000000,Key=Gamepad_LeftX)
++AxisMappings=(AxisName="MotionControllerThumbLeft_Y",Scale=1.000000,Key=Vive_Left_Trackpad_Y)
++AxisMappings=(AxisName="MotionControllerThumbLeft_Y",Scale=1.000000,Key=OculusTouch_Left_Thumbstick_Y)
++AxisMappings=(AxisName="MotionControllerThumbLeft_Y",Scale=1.000000,Key=ValveIndex_Left_Thumbstick_Y)
++AxisMappings=(AxisName="MotionControllerThumbLeft_Y",Scale=1.000000,Key=MixedReality_Left_Thumbstick_Y)
++AxisMappings=(AxisName="MotionControllerThumbLeft_Y",Scale=1.000000,Key=OculusGo_Left_Trackpad_Y)
++AxisMappings=(AxisName="MotionControllerThumbLeft_X",Scale=1.000000,Key=Vive_Left_Trackpad_X)
++AxisMappings=(AxisName="MotionControllerThumbLeft_X",Scale=1.000000,Key=OculusTouch_Left_Thumbstick_X)
++AxisMappings=(AxisName="MotionControllerThumbLeft_X",Scale=1.000000,Key=ValveIndex_Left_Thumbstick_X)
++AxisMappings=(AxisName="MotionControllerThumbLeft_X",Scale=1.000000,Key=MixedReality_Left_Thumbstick_X)
++AxisMappings=(AxisName="MotionControllerThumbLeft_X",Scale=1.000000,Key=OculusGo_Left_Trackpad_X)
++AxisMappings=(AxisName="MotionControllerThumbRight_Y",Scale=1.000000,Key=Vive_Right_Trackpad_Y)
++AxisMappings=(AxisName="MotionControllerThumbRight_Y",Scale=1.000000,Key=OculusTouch_Right_Thumbstick_Y)
++AxisMappings=(AxisName="MotionControllerThumbRight_Y",Scale=1.000000,Key=ValveIndex_Right_Thumbstick_Y)
++AxisMappings=(AxisName="MotionControllerThumbRight_Y",Scale=1.000000,Key=MixedReality_Right_Thumbstick_Y)
++AxisMappings=(AxisName="MotionControllerThumbRight_Y",Scale=1.000000,Key=OculusGo_Right_Trackpad_Y)
++AxisMappings=(AxisName="MotionControllerThumbRight_X",Scale=1.000000,Key=Vive_Right_Trackpad_X)
++AxisMappings=(AxisName="MotionControllerThumbRight_X",Scale=1.000000,Key=OculusTouch_Right_Thumbstick_X)
++AxisMappings=(AxisName="MotionControllerThumbRight_X",Scale=1.000000,Key=ValveIndex_Right_Thumbstick_X)
++AxisMappings=(AxisName="MotionControllerThumbRight_X",Scale=1.000000,Key=MixedReality_Right_Thumbstick_X)
++AxisMappings=(AxisName="MotionControllerThumbRight_X",Scale=1.000000,Key=OculusGo_Right_Trackpad_X)
++AxisMappings=(AxisName="MotionControllerThumbLeft_Y",Scale=1.000000,Key=Daydream_Left_Trackpad_Y)
++AxisMappings=(AxisName="MotionControllerThumbLeft_X",Scale=1.000000,Key=Daydream_Left_Trackpad_X)
++AxisMappings=(AxisName="MotionControllerThumbRight_Y",Scale=1.000000,Key=Daydream_Right_Trackpad_Y)
++AxisMappings=(AxisName="MotionControllerThumbRight_X",Scale=1.000000,Key=Daydream_Right_Trackpad_X)
 DefaultTouchInterface=None
 ConsoleKey=None
 -ConsoleKeys=Tilde
diff --git a/Templates/TP_nDisplayBP/Config/TemplateDefs.ini b/Templates/TP_nDisplayBP/Config/TemplateDefs.ini
--- a/Templates/TP_nDisplayBP/Config/TemplateDefs.ini
+++ b/Templates/TP_nDisplayBP/Config/TemplateDefs.ini
@@ -46,10 +46,10 @@ ReplacementsInFiles=(Extensions=("cpp","h","ini","cs"),From="%TEMPLATENAME%",To=
 ; HardwareTarget - Choose between Mobile and Desktop hardware targets.
 ; GraphicsPreset - Choose between Maximum and Scalable graphics.
 ; StarterContent - Choose whether to add starter content to the project.
-; VR - Choose whether to enable Virtual Reality support for the project.
+; XR - Choose whether to enable Virtual Reality support for the project.
 ;HiddenSettings=All
 HiddenSettings=Languages
 HiddenSettings=HardwareTarget
 HiddenSettings=GraphicsPreset 
 ;HiddenSettings=StarterContent 
-HiddenSettings=VR 
\ No newline at end of file
+HiddenSettings=XR 
\ No newline at end of file
