UE4/Main copy of new VR Template

diff --git a/Templates/TP_VirtualRealityBP/Config/DefaultEngine.ini b/Templates/TP_VirtualRealityBP/Config/DefaultEngine.ini
new file mode 100644
--- /dev/null
+++ b/Templates/TP_VirtualRealityBP/Config/DefaultEngine.ini
@@ -0,0 +1,263 @@
+
+
+[/Script/Engine.RendererSettings]
+r.Mobile.DisableVertexFog=True
+r.Shadow.CSM.MaxMobileCascades=2
+r.MobileMSAA=4
+r.Mobile.AllowDitheredLODTransition=False
+r.Mobile.AllowSoftwareOcclusion=False
+r.Mobile.VirtualTextures=False
+r.DiscardUnusedQuality=False
+r.AllowOcclusionQueries=True
+r.MinScreenRadiusForLights=0.030000
+r.MinScreenRadiusForDepthPrepass=0.030000
+r.MinScreenRadiusForCSMDepth=0.010000
+r.PrecomputedVisibilityWarning=False
+r.TextureStreaming=True
+Compat.UseDXT5NormalMaps=False
+r.VirtualTextures=False
+r.VirtualTexturedLightmaps=False
+r.VT.TileSize=128
+r.VT.TileBorderSize=4
+r.vt.FeedbackFactor=16
+r.VT.EnableCompressZlib=True
+r.VT.EnableCompressCrunch=False
+r.ClearCoatNormal=False
+r.ReflectionCaptureResolution=128
+r.Mobile.ReflectionCaptureCompression=False
+r.ReflectionEnvironmentLightmapMixBasedOnRoughness=True
+r.ForwardShading=True
+r.VertexFoggingForOpaque=True
+r.AllowStaticLighting=True
+r.NormalMapsForStaticLighting=False
+r.GenerateMeshDistanceFields=False
+r.DistanceFieldBuild.EightBit=False
+r.GenerateLandscapeGIData=False
+r.DistanceFieldBuild.Compress=False
+r.TessellationAdaptivePixelsPerTriangle=48.000000
+r.SeparateTranslucency=False
+r.TranslucentSortPolicy=0
+TranslucentSortAxis=(X=0.000000,Y=-1.000000,Z=0.000000)
+r.CustomDepth=1
+r.CustomDepthTemporalAAJitter=True
+r.PostProcessing.PropagateAlpha=0
+r.DefaultFeature.Bloom=False
+r.DefaultFeature.AmbientOcclusion=False
+r.DefaultFeature.AmbientOcclusionStaticFraction=True
+r.DefaultFeature.AutoExposure=False
+r.DefaultFeature.AutoExposure.Method=0
+r.DefaultFeature.AutoExposure.Bias=1.000000
+r.DefaultFeature.AutoExposure.ExtendDefaultLuminanceRange=True
+r.UsePreExposure=True
+r.EyeAdaptation.EditorOnly=False
+r.DefaultFeature.MotionBlur=False
+r.DefaultFeature.LensFlare=False
+r.TemporalAA.Upsampling=False
+r.SSGI.Enable=False
+r.DefaultFeature.AntiAliasing=3
+r.DefaultFeature.LightUnits=1
+r.DefaultBackBufferPixelFormat=4
+r.Shadow.UnbuiltPreviewInGame=True
+r.StencilForLODDither=False
+r.EarlyZPass=3
+r.EarlyZPassOnlyMaterialMasking=False
+r.DBuffer=True
+r.ClearSceneMethod=1
+r.BasePassOutputsVelocity=False
+r.VertexDeformationOutputsVelocity=False
+r.SelectiveBasePassOutputs=False
+bDefaultParticleCutouts=False
+fx.GPUSimulationTextureSizeX=1024
+fx.GPUSimulationTextureSizeY=1024
+r.AllowGlobalClipPlane=False
+r.GBufferFormat=1
+r.MorphTarget.Mode=True
+r.GPUCrashDebugging=False
+vr.InstancedStereo=True
+r.MobileHDR=False
+vr.MobileMultiView=True
+r.Mobile.UseHWsRGBEncoding=True
+vr.RoundRobinOcclusion=False
+vr.ODSCapture=False
+r.MeshStreaming=False
+r.WireframeCullThreshold=5.000000
+r.RayTracing=False
+r.RayTracing.UseTextureLod=False
+r.SupportStationarySkylight=True
+r.SupportLowQualityLightmaps=True
+r.SupportPointLightWholeSceneShadows=True
+r.SupportAtmosphericFog=True
+r.SupportSkyAtmosphere=True
+r.SupportSkyAtmosphereAffectsHeightFog=False
+r.SkinCache.CompileShaders=False
+r.SkinCache.DefaultBehavior=1
+r.SkinCache.SceneMemoryLimitInMB=128.000000
+r.Mobile.EnableStaticAndCSMShadowReceivers=True
+r.Mobile.EnableMovableLightCSMShaderCulling=True
+r.Mobile.AllowDistanceFieldShadows=True
+r.Mobile.AllowMovableDirectionalLights=True
+r.MobileNumDynamicPointLights=4
+r.MobileDynamicPointLightsUseStaticBranch=True
+r.Mobile.EnableMovableSpotlights=False
+r.Mobile.EnableMovableSpotlightsShadow=False
+r.GPUSkin.Support16BitBoneIndex=False
+r.GPUSkin.Limit2BoneInfluences=False
+r.SupportDepthOnlyIndexBuffers=True
+r.SupportReversedIndexBuffers=True
+r.LightPropagationVolume=False
+r.Mobile.AmbientOcclusion=False
+r.GPUSkin.UnlimitedBoneInfluences=False
+r.GPUSkin.UnlimitedBoneInfluencesThreshold=8
+r.Mobile.PlanarReflectionMode=0
+bStreamSkeletalMeshLODs=(Default=False,PerPlatform=())
+bDiscardSkeletalMeshOptionalLODs=(Default=False,PerPlatform=())
+VisualizeCalibrationColorMaterialPath=None
+VisualizeCalibrationCustomMaterialPath=None
+VisualizeCalibrationGrayscaleMaterialPath=None
+r.Mobile.SupportGPUScene=1
+
+[/Script/HardwareTargeting.HardwareTargetingSettings]
+TargetedHardwareClass=Mobile
+AppliedTargetedHardwareClass=Mobile
+DefaultGraphicsPerformance=Scalable
+AppliedDefaultGraphicsPerformance=Scalable
+
+[/Script/EngineSettings.GameMapsSettings]
+EditorStartupMap=/Game/VRTemplate/Maps/VRTemplateMap.VRTemplateMap
+LocalMapOptions=
+TransitionMap=None
+bUseSplitscreen=False
+TwoPlayerSplitscreenLayout=Horizontal
+ThreePlayerSplitscreenLayout=FavorTop
+FourPlayerSplitscreenLayout=Grid
+bOffsetPlayerGamepadIds=False
+GameInstanceClass=/Script/Engine.GameInstance
+GameDefaultMap=/Game/VRTemplate/Maps/VRTemplateMap.VRTemplateMap
+ServerDefaultMap=/Engine/Maps/Entry.Entry
+GlobalDefaultGameMode=/Game/VRTemplate/Blueprints/VRGameMode.VRGameMode_C
+GlobalDefaultServerGameMode=None
+
+[/Script/Slate.SlateSettings]
+bExplicitCanvasChildZOrder=True
+
+[/Script/WindowsTargetPlatform.WindowsTargetSettings]
+SpatializationPlugin=Resonance Audio
+ReverbPlugin=Resonance Audio
+
+[/Script/AndroidRuntimeSettings.AndroidRuntimeSettings]
+PackageName=com.epicgames.[PROJECT]
+StoreVersion=1
+StoreVersionOffsetArmV7=0
+StoreVersionOffsetArm64=0
+StoreVersionOffsetX8664=0
+ApplicationDisplayName=
+VersionDisplayName=1.0
+MinSDKVersion=23
+TargetSDKVersion=25
+InstallLocation=InternalOnly
+bEnableGradle=True
+bEnableLint=False
+bPackageDataInsideApk=True
+bCreateAllPlatformsInstall=False
+bDisableVerifyOBBOnStartUp=False
+bForceSmallOBBFiles=False
+bAllowLargeOBBFiles=False
+bAllowPatchOBBFile=False
+bAllowOverflowOBBFiles=False
+bUseExternalFilesDir=False
+bPublicLogFiles=True
+Orientation=SensorLandscape
+MaxAspectRatio=2.100000
+bUseDisplayCutout=False
+bRestoreNotificationsOnReboot=False
+bFullScreen=True
+bEnableNewKeyboard=True
+DepthBufferPreference=Default
+bValidateTextureFormats=True
+bEnableBundle=False
+bEnableUniversalAPK=True
+bBundleABISplit=True
+bBundleLanguageSplit=True
+bBundleDensitySplit=True
+ExtraApplicationSettings=
+ExtraActivitySettings=
+bAndroidVoiceEnabled=False
++PackageForOculusMobile=Quest2
++PackageForOculusMobile=Quest
+bRemoveOSIG=False
++GoogleVRCaps=Daydream33
+bGoogleVRSustainedPerformance=False
+KeyStore=UE.keystore
+KeyAlias=DistroKey
+KeyStorePassword=codesigner
+KeyPassword=
+bBuildForArmV7=True
+bBuildForArm64=False
+bBuildForX8664=False
+bBuildForES31=True
+bSupportsVulkan=True
+bSupportsVulkanSM5=False
+ClangSanitizer=None
+bAndroidOpenGLSupportsBackbufferSampling=False
+bDetectVulkanByDefault=True
+bBuildWithHiddenSymbolVisibility=False
+bSaveSymbols=False
+bForceLDLinker=False
+bEnableGooglePlaySupport=False
+bUseGetAccounts=False
+GamesAppID=
+bEnableSnapshots=False
+bSupportAdMob=True
+AdMobAdUnitID=
+GooglePlayLicenseKey=
+GCMClientSenderID=
+bShowLaunchImage=True
+bAllowIMU=True
+bAllowControllers=True
+bBlockAndroidKeysOnControllers=False
+bControllersBlockDeviceFeedback=False
+AndroidAudio=Default
+AudioSampleRate=44100
+AudioCallbackBufferFrameSize=1024
+AudioNumBuffersToEnqueue=4
+AudioMaxChannels=0
+AudioNumSourceWorkers=0
+SpatializationPlugin=Resonance Audio
+ReverbPlugin=Resonance Audio
+OcclusionPlugin=
+CompressionOverrides=(bOverrideCompressionTimes=False,DurationThreshold=5.000000,MaxNumRandomBranches=0,SoundCueQualityIndex=0)
+bUseAudioStreamCaching=False
+CacheSizeKB=0
+MaxChunkSizeOverrideKB=0
+bResampleForDevice=False
+SoundCueCookQualityIndex=-1
+MaxSampleRate=0.000000
+HighSampleRate=0.000000
+MedSampleRate=0.000000
+LowSampleRate=0.000000
+MinSampleRate=0.000000
+CompressionQualityModifier=0.000000
+AutoStreamingThreshold=0.000000
+AndroidGraphicsDebugger=None
+MaliGraphicsDebuggerPath=(Path="")
+bEnableMaliPerfCounters=False
+bMultiTargetFormat_ETC2=True
+bMultiTargetFormat_DXT=True
+bMultiTargetFormat_ASTC=True
+TextureFormatPriority_ETC2=0.200000
+TextureFormatPriority_DXT=0.600000
+TextureFormatPriority_ASTC=0.900000
+SDKAPILevelOverride=
+NDKAPILevelOverride=
+bStreamLandscapeMeshLODs=False
+
+[/Script/IOSRuntimeSettings.IOSRuntimeSettings]
+SpatializationPlugin=Resonance Audio
+
+[/Script/LuminRuntimeSettings.LuminRuntimeSettings]
+SpatializationPlugin=Built-in Spatialization
+
+[/Script/MacTargetPlatform.MacTargetSettings]
+SpatializationPlugin=Resonance Audio
+ReverbPlugin=Resonance Audio
+
diff --git a/Templates/TP_VirtualRealityBP/Config/DefaultGame.ini b/Templates/TP_VirtualRealityBP/Config/DefaultGame.ini
new file mode 100644
--- /dev/null
+++ b/Templates/TP_VirtualRealityBP/Config/DefaultGame.ini
@@ -0,0 +1,4 @@
+[/Script/EngineSettings.GeneralProjectSettings]
+ProjectID=023E2CFB43C311B913BF9F82BAC14D25
+bStartInVR=True
+
diff --git a/Templates/TP_VirtualRealityBP/Config/DefaultInput.ini b/Templates/TP_VirtualRealityBP/Config/DefaultInput.ini
new file mode 100644
--- /dev/null
+++ b/Templates/TP_VirtualRealityBP/Config/DefaultInput.ini
@@ -0,0 +1,164 @@
+[/Script/Engine.InputSettings]
+-AxisConfig=(AxisKeyName="Gamepad_LeftX",AxisProperties=(DeadZone=0.25,Exponent=1.f,Sensitivity=1.f))
+-AxisConfig=(AxisKeyName="Gamepad_LeftY",AxisProperties=(DeadZone=0.25,Exponent=1.f,Sensitivity=1.f))
+-AxisConfig=(AxisKeyName="Gamepad_RightX",AxisProperties=(DeadZone=0.25,Exponent=1.f,Sensitivity=1.f))
+-AxisConfig=(AxisKeyName="Gamepad_RightY",AxisProperties=(DeadZone=0.25,Exponent=1.f,Sensitivity=1.f))
+-AxisConfig=(AxisKeyName="MouseX",AxisProperties=(DeadZone=0.f,Exponent=1.f,Sensitivity=0.07f))
+-AxisConfig=(AxisKeyName="MouseY",AxisProperties=(DeadZone=0.f,Exponent=1.f,Sensitivity=0.07f))
+-AxisConfig=(AxisKeyName="Mouse2D",AxisProperties=(DeadZone=0.f,Exponent=1.f,Sensitivity=0.07f))
++AxisConfig=(AxisKeyName="Gamepad_LeftX",AxisProperties=(DeadZone=0.250000,Sensitivity=1.000000,Exponent=1.000000,bInvert=False))
++AxisConfig=(AxisKeyName="Gamepad_LeftY",AxisProperties=(DeadZone=0.250000,Sensitivity=1.000000,Exponent=1.000000,bInvert=False))
++AxisConfig=(AxisKeyName="Gamepad_RightX",AxisProperties=(DeadZone=0.250000,Sensitivity=1.000000,Exponent=1.000000,bInvert=False))
++AxisConfig=(AxisKeyName="Gamepad_RightY",AxisProperties=(DeadZone=0.250000,Sensitivity=1.000000,Exponent=1.000000,bInvert=False))
++AxisConfig=(AxisKeyName="MouseX",AxisProperties=(DeadZone=0.000000,Sensitivity=0.070000,Exponent=1.000000,bInvert=False))
++AxisConfig=(AxisKeyName="MouseY",AxisProperties=(DeadZone=0.000000,Sensitivity=0.070000,Exponent=1.000000,bInvert=False))
++AxisConfig=(AxisKeyName="Mouse2D",AxisProperties=(DeadZone=0.000000,Sensitivity=0.070000,Exponent=1.000000,bInvert=False))
++AxisConfig=(AxisKeyName="MouseWheelAxis",AxisProperties=(DeadZone=0.000000,Sensitivity=1.000000,Exponent=1.000000,bInvert=False))
++AxisConfig=(AxisKeyName="Gamepad_LeftTriggerAxis",AxisProperties=(DeadZone=0.000000,Sensitivity=1.000000,Exponent=1.000000,bInvert=False))
++AxisConfig=(AxisKeyName="Gamepad_RightTriggerAxis",AxisProperties=(DeadZone=0.000000,Sensitivity=1.000000,Exponent=1.000000,bInvert=False))
++AxisConfig=(AxisKeyName="Gamepad_Special_Left_X",AxisProperties=(DeadZone=0.000000,Sensitivity=1.000000,Exponent=1.000000,bInvert=False))
++AxisConfig=(AxisKeyName="Gamepad_Special_Left_Y",AxisProperties=(DeadZone=0.000000,Sensitivity=1.000000,Exponent=1.000000,bInvert=False))
++AxisConfig=(AxisKeyName="Daydream_Left_Trackpad_X",AxisProperties=(DeadZone=0.000000,Sensitivity=1.000000,Exponent=1.000000,bInvert=False))
++AxisConfig=(AxisKeyName="Daydream_Left_Trackpad_Y",AxisProperties=(DeadZone=0.000000,Sensitivity=1.000000,Exponent=1.000000,bInvert=False))
++AxisConfig=(AxisKeyName="Daydream_Right_Trackpad_X",AxisProperties=(DeadZone=0.000000,Sensitivity=1.000000,Exponent=1.000000,bInvert=False))
++AxisConfig=(AxisKeyName="Daydream_Right_Trackpad_Y",AxisProperties=(DeadZone=0.000000,Sensitivity=1.000000,Exponent=1.000000,bInvert=False))
++AxisConfig=(AxisKeyName="Vive_Left_Trigger_Axis",AxisProperties=(DeadZone=0.000000,Sensitivity=1.000000,Exponent=1.000000,bInvert=False))
++AxisConfig=(AxisKeyName="Vive_Left_Trackpad_X",AxisProperties=(DeadZone=0.000000,Sensitivity=1.000000,Exponent=1.000000,bInvert=False))
++AxisConfig=(AxisKeyName="Vive_Left_Trackpad_Y",AxisProperties=(DeadZone=0.000000,Sensitivity=1.000000,Exponent=1.000000,bInvert=False))
++AxisConfig=(AxisKeyName="Vive_Right_Trigger_Axis",AxisProperties=(DeadZone=0.000000,Sensitivity=1.000000,Exponent=1.000000,bInvert=False))
++AxisConfig=(AxisKeyName="Vive_Right_Trackpad_X",AxisProperties=(DeadZone=0.000000,Sensitivity=1.000000,Exponent=1.000000,bInvert=False))
++AxisConfig=(AxisKeyName="Vive_Right_Trackpad_Y",AxisProperties=(DeadZone=0.000000,Sensitivity=1.000000,Exponent=1.000000,bInvert=False))
++AxisConfig=(AxisKeyName="MixedReality_Left_Trigger_Axis",AxisProperties=(DeadZone=0.000000,Sensitivity=1.000000,Exponent=1.000000,bInvert=False))
++AxisConfig=(AxisKeyName="MixedReality_Left_Thumbstick_X",AxisProperties=(DeadZone=0.000000,Sensitivity=1.000000,Exponent=1.000000,bInvert=False))
++AxisConfig=(AxisKeyName="MixedReality_Left_Thumbstick_Y",AxisProperties=(DeadZone=0.000000,Sensitivity=1.000000,Exponent=1.000000,bInvert=False))
++AxisConfig=(AxisKeyName="MixedReality_Left_Trackpad_X",AxisProperties=(DeadZone=0.000000,Sensitivity=1.000000,Exponent=1.000000,bInvert=False))
++AxisConfig=(AxisKeyName="MixedReality_Left_Trackpad_Y",AxisProperties=(DeadZone=0.000000,Sensitivity=1.000000,Exponent=1.000000,bInvert=False))
++AxisConfig=(AxisKeyName="MixedReality_Right_Trigger_Axis",AxisProperties=(DeadZone=0.000000,Sensitivity=1.000000,Exponent=1.000000,bInvert=False))
++AxisConfig=(AxisKeyName="MixedReality_Right_Thumbstick_X",AxisProperties=(DeadZone=0.000000,Sensitivity=1.000000,Exponent=1.000000,bInvert=False))
++AxisConfig=(AxisKeyName="MixedReality_Right_Thumbstick_Y",AxisProperties=(DeadZone=0.000000,Sensitivity=1.000000,Exponent=1.000000,bInvert=False))
++AxisConfig=(AxisKeyName="MixedReality_Right_Trackpad_X",AxisProperties=(DeadZone=0.000000,Sensitivity=1.000000,Exponent=1.000000,bInvert=False))
++AxisConfig=(AxisKeyName="MixedReality_Right_Trackpad_Y",AxisProperties=(DeadZone=0.000000,Sensitivity=1.000000,Exponent=1.000000,bInvert=False))
++AxisConfig=(AxisKeyName="OculusTouch_Left_Grip_Axis",AxisProperties=(DeadZone=0.000000,Sensitivity=1.000000,Exponent=1.000000,bInvert=False))
++AxisConfig=(AxisKeyName="OculusTouch_Left_Trigger_Axis",AxisProperties=(DeadZone=0.000000,Sensitivity=1.000000,Exponent=1.000000,bInvert=False))
++AxisConfig=(AxisKeyName="OculusTouch_Left_Thumbstick_X",AxisProperties=(DeadZone=0.000000,Sensitivity=1.000000,Exponent=1.000000,bInvert=False))
++AxisConfig=(AxisKeyName="OculusTouch_Left_Thumbstick_Y",AxisProperties=(DeadZone=0.000000,Sensitivity=1.000000,Exponent=1.000000,bInvert=False))
++AxisConfig=(AxisKeyName="OculusTouch_Right_Grip_Axis",AxisProperties=(DeadZone=0.000000,Sensitivity=1.000000,Exponent=1.000000,bInvert=False))
++AxisConfig=(AxisKeyName="OculusTouch_Right_Trigger_Axis",AxisProperties=(DeadZone=0.000000,Sensitivity=1.000000,Exponent=1.000000,bInvert=False))
++AxisConfig=(AxisKeyName="OculusTouch_Right_Thumbstick_X",AxisProperties=(DeadZone=0.000000,Sensitivity=1.000000,Exponent=1.000000,bInvert=False))
++AxisConfig=(AxisKeyName="OculusTouch_Right_Thumbstick_Y",AxisProperties=(DeadZone=0.000000,Sensitivity=1.000000,Exponent=1.000000,bInvert=False))
++AxisConfig=(AxisKeyName="ValveIndex_Left_Grip_Axis",AxisProperties=(DeadZone=0.000000,Sensitivity=1.000000,Exponent=1.000000,bInvert=False))
++AxisConfig=(AxisKeyName="ValveIndex_Left_Grip_Force",AxisProperties=(DeadZone=0.000000,Sensitivity=1.000000,Exponent=1.000000,bInvert=False))
++AxisConfig=(AxisKeyName="ValveIndex_Left_Trigger_Axis",AxisProperties=(DeadZone=0.000000,Sensitivity=1.000000,Exponent=1.000000,bInvert=False))
++AxisConfig=(AxisKeyName="ValveIndex_Left_Thumbstick_X",AxisProperties=(DeadZone=0.000000,Sensitivity=1.000000,Exponent=1.000000,bInvert=False))
++AxisConfig=(AxisKeyName="ValveIndex_Left_Thumbstick_Y",AxisProperties=(DeadZone=0.000000,Sensitivity=1.000000,Exponent=1.000000,bInvert=False))
++AxisConfig=(AxisKeyName="ValveIndex_Left_Trackpad_X",AxisProperties=(DeadZone=0.000000,Sensitivity=1.000000,Exponent=1.000000,bInvert=False))
++AxisConfig=(AxisKeyName="ValveIndex_Left_Trackpad_Y",AxisProperties=(DeadZone=0.000000,Sensitivity=1.000000,Exponent=1.000000,bInvert=False))
++AxisConfig=(AxisKeyName="ValveIndex_Left_Trackpad_Force",AxisProperties=(DeadZone=0.000000,Sensitivity=1.000000,Exponent=1.000000,bInvert=False))
++AxisConfig=(AxisKeyName="ValveIndex_Left_Trackpad_Touch",AxisProperties=(DeadZone=0.000000,Sensitivity=1.000000,Exponent=1.000000,bInvert=False))
++AxisConfig=(AxisKeyName="ValveIndex_Right_Grip_Axis",AxisProperties=(DeadZone=0.000000,Sensitivity=1.000000,Exponent=1.000000,bInvert=False))
++AxisConfig=(AxisKeyName="ValveIndex_Right_Grip_Force",AxisProperties=(DeadZone=0.000000,Sensitivity=1.000000,Exponent=1.000000,bInvert=False))
++AxisConfig=(AxisKeyName="ValveIndex_Right_Trigger_Axis",AxisProperties=(DeadZone=0.000000,Sensitivity=1.000000,Exponent=1.000000,bInvert=False))
++AxisConfig=(AxisKeyName="ValveIndex_Right_Thumbstick_X",AxisProperties=(DeadZone=0.000000,Sensitivity=1.000000,Exponent=1.000000,bInvert=False))
++AxisConfig=(AxisKeyName="ValveIndex_Right_Thumbstick_Y",AxisProperties=(DeadZone=0.000000,Sensitivity=1.000000,Exponent=1.000000,bInvert=False))
++AxisConfig=(AxisKeyName="ValveIndex_Right_Trackpad_X",AxisProperties=(DeadZone=0.000000,Sensitivity=1.000000,Exponent=1.000000,bInvert=False))
++AxisConfig=(AxisKeyName="ValveIndex_Right_Trackpad_Y",AxisProperties=(DeadZone=0.000000,Sensitivity=1.000000,Exponent=1.000000,bInvert=False))
++AxisConfig=(AxisKeyName="ValveIndex_Right_Trackpad_Force",AxisProperties=(DeadZone=0.000000,Sensitivity=1.000000,Exponent=1.000000,bInvert=False))
++AxisConfig=(AxisKeyName="MotionController_Left_Thumbstick_Z",AxisProperties=(DeadZone=0.000000,Sensitivity=1.000000,Exponent=1.000000,bInvert=False))
++AxisConfig=(AxisKeyName="MagicLeap_Left_Trigger_Axis",AxisProperties=(DeadZone=0.000000,Sensitivity=1.000000,Exponent=1.000000,bInvert=False))
++AxisConfig=(AxisKeyName="MagicLeap_Left_Trackpad_X",AxisProperties=(DeadZone=0.000000,Sensitivity=1.000000,Exponent=1.000000,bInvert=False))
++AxisConfig=(AxisKeyName="MagicLeap_Left_Trackpad_Y",AxisProperties=(DeadZone=0.000000,Sensitivity=1.000000,Exponent=1.000000,bInvert=False))
++AxisConfig=(AxisKeyName="MagicLeap_Left_Trackpad_Force",AxisProperties=(DeadZone=0.000000,Sensitivity=1.000000,Exponent=1.000000,bInvert=False))
++AxisConfig=(AxisKeyName="MagicLeap_Left_Touch1_X",AxisProperties=(DeadZone=0.000000,Sensitivity=1.000000,Exponent=1.000000,bInvert=False))
++AxisConfig=(AxisKeyName="MagicLeap_Left_Touch1_Y",AxisProperties=(DeadZone=0.000000,Sensitivity=1.000000,Exponent=1.000000,bInvert=False))
++AxisConfig=(AxisKeyName="MagicLeap_Left_Touch1_Force",AxisProperties=(DeadZone=0.000000,Sensitivity=1.000000,Exponent=1.000000,bInvert=False))
++AxisConfig=(AxisKeyName="MotionController_Right_Thumbstick_Z",AxisProperties=(DeadZone=0.000000,Sensitivity=1.000000,Exponent=1.000000,bInvert=False))
++AxisConfig=(AxisKeyName="MagicLeap_Right_Trigger_Axis",AxisProperties=(DeadZone=0.000000,Sensitivity=1.000000,Exponent=1.000000,bInvert=False))
++AxisConfig=(AxisKeyName="MagicLeap_Right_Trackpad_X",AxisProperties=(DeadZone=0.000000,Sensitivity=1.000000,Exponent=1.000000,bInvert=False))
++AxisConfig=(AxisKeyName="MagicLeap_Right_Trackpad_Y",AxisProperties=(DeadZone=0.000000,Sensitivity=1.000000,Exponent=1.000000,bInvert=False))
++AxisConfig=(AxisKeyName="MagicLeap_Right_Trackpad_Force",AxisProperties=(DeadZone=0.000000,Sensitivity=1.000000,Exponent=1.000000,bInvert=False))
++AxisConfig=(AxisKeyName="MagicLeap_Right_Touch1_X",AxisProperties=(DeadZone=0.000000,Sensitivity=1.000000,Exponent=1.000000,bInvert=False))
++AxisConfig=(AxisKeyName="MagicLeap_Right_Touch1_Y",AxisProperties=(DeadZone=0.000000,Sensitivity=1.000000,Exponent=1.000000,bInvert=False))
++AxisConfig=(AxisKeyName="MagicLeap_Right_Touch1_Force",AxisProperties=(DeadZone=0.000000,Sensitivity=1.000000,Exponent=1.000000,bInvert=False))
+bAltEnterTogglesFullscreen=True
+bF11TogglesFullscreen=True
+bUseMouseForTouch=False
+bEnableMouseSmoothing=True
+bEnableFOVScaling=True
+bCaptureMouseOnLaunch=True
+bAlwaysShowTouchInterface=False
+bShowConsoleOnFourFingerTap=True
+bEnableGestureRecognizer=False
+bUseAutocorrect=False
+DefaultViewportMouseCaptureMode=CapturePermanently_IncludingInitialMouseDown
+DefaultViewportMouseLockMode=LockOnCapture
+FOVScale=0.011110
+DoubleClickTime=0.200000
++ActionMappings=(ActionName="GrabLeft",bShift=False,bCtrl=False,bAlt=False,bCmd=False,Key=OculusTouch_Left_Grip_Click)
++ActionMappings=(ActionName="GrabRight",bShift=False,bCtrl=False,bAlt=False,bCmd=False,Key=OculusTouch_Right_Grip_Click)
++ActionMappings=(ActionName="TriggerLeft",bShift=False,bCtrl=False,bAlt=False,bCmd=False,Key=OculusTouch_Left_Trigger_Click)
++ActionMappings=(ActionName="TriggerRight",bShift=False,bCtrl=False,bAlt=False,bCmd=False,Key=OculusTouch_Right_Trigger_Click)
++ActionMappings=(ActionName="GrabLeft",bShift=False,bCtrl=False,bAlt=False,bCmd=False,Key=ValveIndex_Left_Grip_Axis)
++ActionMappings=(ActionName="GrabLeft",bShift=False,bCtrl=False,bAlt=False,bCmd=False,Key=MixedReality_Left_Grip_Click)
++ActionMappings=(ActionName="GrabLeft",bShift=False,bCtrl=False,bAlt=False,bCmd=False,Key=Vive_Left_Grip_Click)
++ActionMappings=(ActionName="GrabRight",bShift=False,bCtrl=False,bAlt=False,bCmd=False,Key=ValveIndex_Right_Grip_Axis)
++ActionMappings=(ActionName="GrabRight",bShift=False,bCtrl=False,bAlt=False,bCmd=False,Key=MixedReality_Right_Grip_Click)
++ActionMappings=(ActionName="GrabRight",bShift=False,bCtrl=False,bAlt=False,bCmd=False,Key=Vive_Right_Grip_Click)
++ActionMappings=(ActionName="TriggerLeft",bShift=False,bCtrl=False,bAlt=False,bCmd=False,Key=ValveIndex_Left_Trigger_Click)
++ActionMappings=(ActionName="TriggerLeft",bShift=False,bCtrl=False,bAlt=False,bCmd=False,Key=MixedReality_Left_Trigger_Click)
++ActionMappings=(ActionName="TriggerLeft",bShift=False,bCtrl=False,bAlt=False,bCmd=False,Key=Vive_Left_Trigger_Click)
++ActionMappings=(ActionName="TriggerRight",bShift=False,bCtrl=False,bAlt=False,bCmd=False,Key=ValveIndex_Right_Trigger_Click)
++ActionMappings=(ActionName="TriggerRight",bShift=False,bCtrl=False,bAlt=False,bCmd=False,Key=MixedReality_Right_Trigger_Click)
++ActionMappings=(ActionName="TriggerRight",bShift=False,bCtrl=False,bAlt=False,bCmd=False,Key=Vive_Right_Trigger_Click)
++ActionMappings=(ActionName="SpectatorResetFOV",bShift=False,bCtrl=False,bAlt=False,bCmd=False,Key=MiddleMouseButton)
++ActionMappings=(ActionName="SpectatorFadeToggle",bShift=False,bCtrl=False,bAlt=False,bCmd=False,Key=F)
++ActionMappings=(ActionName="SpectatorActiveToggle",bShift=False,bCtrl=False,bAlt=False,bCmd=False,Key=Tab)
++ActionMappings=(ActionName="SpectatorResetRotation",bShift=False,bCtrl=False,bAlt=False,bCmd=False,Key=R)
++AxisMappings=(AxisName="MovementAxisLeft_X",Scale=1.000000,Key=OculusTouch_Left_Thumbstick_X)
++AxisMappings=(AxisName="GrabAxisLeft",Scale=1.000000,Key=OculusTouch_Left_Grip_Axis)
++AxisMappings=(AxisName="MovementAxisLeft_Y",Scale=1.000000,Key=OculusTouch_Left_Thumbstick_Y)
++AxisMappings=(AxisName="TriggerAxisLeft",Scale=1.000000,Key=OculusTouch_Left_Trigger_Axis)
++AxisMappings=(AxisName="TriggerAxisRight",Scale=1.000000,Key=OculusTouch_Right_Trigger_Axis)
++AxisMappings=(AxisName="MovementAxisLeft_X",Scale=1.000000,Key=ValveIndex_Left_Thumbstick_X)
++AxisMappings=(AxisName="MovementAxisLeft_X",Scale=1.000000,Key=MixedReality_Left_Thumbstick_X)
++AxisMappings=(AxisName="MovementAxisLeft_X",Scale=1.000000,Key=Vive_Left_Trackpad_X)
++AxisMappings=(AxisName="GrabAxisLeft",Scale=1.000000,Key=ValveIndex_Left_Grip_Axis)
++AxisMappings=(AxisName="GrabAxisLeft",Scale=1.000000,Key=MixedReality_Left_Grip_Click)
++AxisMappings=(AxisName="GrabAxisLeft",Scale=1.000000,Key=Vive_Left_Grip_Click)
++AxisMappings=(AxisName="GrabAxisRight",Scale=1.000000,Key=OculusTouch_Right_Grip_Axis)
++AxisMappings=(AxisName="GrabAxisRight",Scale=1.000000,Key=ValveIndex_Right_Grip_Axis)
++AxisMappings=(AxisName="GrabAxisRight",Scale=1.000000,Key=MixedReality_Right_Grip_Click)
++AxisMappings=(AxisName="GrabAxisRight",Scale=1.000000,Key=Vive_Right_Grip_Click)
++AxisMappings=(AxisName="TriggerAxisLeft",Scale=1.000000,Key=ValveIndex_Left_Trigger_Axis)
++AxisMappings=(AxisName="TriggerAxisLeft",Scale=1.000000,Key=MixedReality_Left_Trigger_Axis)
++AxisMappings=(AxisName="TriggerAxisLeft",Scale=1.000000,Key=Vive_Left_Trigger_Axis)
++AxisMappings=(AxisName="TriggerAxisRight",Scale=1.000000,Key=ValveIndex_Right_Trigger_Axis)
++AxisMappings=(AxisName="TriggerAxisRight",Scale=1.000000,Key=MixedReality_Right_Trigger_Axis)
++AxisMappings=(AxisName="TriggerAxisRight",Scale=1.000000,Key=Vive_Right_Trigger_Axis)
++AxisMappings=(AxisName="MovementAxisLeft_Y",Scale=1.000000,Key=ValveIndex_Left_Thumbstick_Y)
++AxisMappings=(AxisName="MovementAxisLeft_Y",Scale=1.000000,Key=MixedReality_Left_Thumbstick_Y)
++AxisMappings=(AxisName="MovementAxisLeft_Y",Scale=1.000000,Key=Vive_Left_Trackpad_Y)
++AxisMappings=(AxisName="SpectatorYaw",Scale=1.000000,Key=MouseX)
++AxisMappings=(AxisName="SpectatorPitch",Scale=1.000000,Key=MouseY)
++AxisMappings=(AxisName="SpectatorFOV",Scale=1.000000,Key=MouseWheelAxis)
++AxisMappings=(AxisName="SpectatorMoveForward",Scale=1.000000,Key=W)
++AxisMappings=(AxisName="SpectatorMoveRight",Scale=1.000000,Key=D)
++AxisMappings=(AxisName="SpectatorMoveForward",Scale=-1.000000,Key=S)
++AxisMappings=(AxisName="SpectatorMoveRight",Scale=-1.000000,Key=A)
++AxisMappings=(AxisName="SpectatorMoveUp",Scale=1.000000,Key=SpaceBar)
++AxisMappings=(AxisName="SpectatorMoveUp",Scale=1.000000,Key=E)
++AxisMappings=(AxisName="SpectatorMoveUp",Scale=-1.000000,Key=Q)
++AxisMappings=(AxisName="MovementAxisRight_X",Scale=1.000000,Key=OculusTouch_Right_Thumbstick_X)
++AxisMappings=(AxisName="MovementAxisRight_X",Scale=1.000000,Key=ValveIndex_Right_Thumbstick_X)
++AxisMappings=(AxisName="MovementAxisRight_X",Scale=1.000000,Key=MixedReality_Right_Thumbstick_X)
++AxisMappings=(AxisName="MovementAxisRight_X",Scale=1.000000,Key=Vive_Right_Trackpad_X)
++AxisMappings=(AxisName="MovementAxisRight_Y",Scale=1.000000,Key=OculusTouch_Right_Thumbstick_Y)
++AxisMappings=(AxisName="MovementAxisRight_Y",Scale=1.000000,Key=ValveIndex_Right_Thumbstick_Y)
++AxisMappings=(AxisName="MovementAxisRight_Y",Scale=1.000000,Key=MixedReality_Right_Thumbstick_Y)
++AxisMappings=(AxisName="MovementAxisRight_Y",Scale=1.000000,Key=Vive_Right_Trackpad_Y)
++AxisMappings=(AxisName="MovementAxisRight_Y",Scale=1.000000,Key=OculusTouch_Right_A_Click)
+DefaultPlayerInputClass=/Script/Engine.PlayerInput
+DefaultInputComponentClass=/Script/Engine.InputComponent
+DefaultTouchInterface=None
+-ConsoleKeys=Tilde
++ConsoleKeys=Tilde
+
diff --git a/Templates/TP_VirtualRealityBP/Config/TemplateDefs.ini b/Templates/TP_VirtualRealityBP/Config/TemplateDefs.ini
new file mode 100644
--- /dev/null
+++ b/Templates/TP_VirtualRealityBP/Config/TemplateDefs.ini
@@ -0,0 +1,54 @@
+﻿[/Script/GameProjectGeneration.TemplateProjectDefs]
+
+Categories=Games
+
+LocalizedDisplayNames=(Language="en",Text="Virtual Reality")
+LocalizedDescriptions=(Language="en",Text="Blueprint Virtual Reality Template using OpenXR for Desktop, Console, and Mobile. The template features teleport locomotion, snap turning, grabbable objects, grabbable and interactive objects, Google Resonance spatialized audio, and a VR Spectator Camera.")
+
+LocalizedDisplayNames=(Language="ko",Text="가상 현실")
+LocalizedDescriptions=(Language="ko",Text="데스크톱 및 Playstation VR 용 블루프린트 가상 현실 템플릿입니다. 각 플랫폼 별 전용 폰과 맵은 물론 두 가지 텔레포트 이동 방식이 포함되어 있습니다. 플레이어는 MotionControllerMap 에서 모션 컨트롤러를 사용하여 텔레포트할 수 있으며, HMDLocomotionMap 에서 게임패드 아니면 HeadMountedDisplay 만으로 이동할 곳을 지정하는 방식으로 텔레포트할 수 있습니다.")
+
+LocalizedDisplayNames=(Language="ja",Text="バーチャルリアリティ")
+LocalizedDescriptions=(Language="ja",Text="デスクトップ VR と Playstation VR 向け、ブループリントのバーチャルリアリティテンプレートです。テンプレートには二種類のテレポートロコモーションモードがあります。それぞれに ポーンとマップが用意されています。MotionControllerMap ではモーションコントローラーを使ったテレポートが可能です。HMDLocomotionMap ではゲームパッドを使ったテレポート、及び、ヘッドマウントディスプレイを使ったポイントアンドテレポートが可能です。")
+
+LocalizedDisplayNames=(Language="zh-Hans",Text="虚拟现实应用")
+LocalizedDescriptions=(Language="zh-Hans",Text="蓝图虚拟现实应用模板面向台式PC以及Playstation VR设备。模板包含2中传送移动模式，分别拥有独立的Pawn和地图。MotionControllerMap可以让玩家使用体感手柄四处传送，而HMDLocomotionMap可以让玩家使用游戏手柄或头显设备进行“指向-传送”式的传送。")
+
+ClassTypes=Pawn, GameMode, Blueprint Interface, Actor
+AssetTypes=Sound Cue, Haptic Feedback Effect Curve, Niagara System.
+
+FoldersToIgnore=Binaries
+FoldersToIgnore=Build
+FoldersToIgnore=Intermediate
+FoldersToIgnore=Saved
+FoldersToIgnore=Media
+
+FilesToIgnore="%TEMPLATENAME%.uproject"
+FilesToIgnore="%TEMPLATENAME%.png"
+FilesToIgnore="Config/TemplateDefs.ini"
+FilesToIgnore="Config/config.ini"
+FilesToIgnore="Manifest.json"
+FilesToIgnore="contents.txt"
+
+FolderRenames=(From="Source/%TEMPLATENAME%",To="Source/%PROJECTNAME%")
+FolderRenames=(From="Source/%TEMPLATENAME%Editor",To="Source/%PROJECTNAME%Editor")
+
+FilenameReplacements=(Extensions=("cpp","h","ini","cs"),From="%TEMPLATENAME_UPPERCASE%",To="%PROJECTNAME_UPPERCASE%",bCaseSensitive=true)
+FilenameReplacements=(Extensions=("cpp","h","ini","cs"),From="%TEMPLATENAME_LOWERCASE%",To="%PROJECTNAME_LOWERCASE%",bCaseSensitive=true)
+FilenameReplacements=(Extensions=("cpp","h","ini","cs"),From="%TEMPLATENAME%",To="%PROJECTNAME%",bCaseSensitive=false)
+
+ReplacementsInFiles=(Extensions=("cpp","h","ini","cs"),From="%TEMPLATENAME_UPPERCASE%",To="%PROJECTNAME_UPPERCASE%",bCaseSensitive=true)
+ReplacementsInFiles=(Extensions=("cpp","h","ini","cs"),From="%TEMPLATENAME_LOWERCASE%",To="%PROJECTNAME_LOWERCASE%",bCaseSensitive=true)
+ReplacementsInFiles=(Extensions=("cpp","h","ini","cs"),From="%TEMPLATENAME%",To="%PROJECTNAME%",bCaseSensitive=false)
+SortKey=VirtualRealityBP
+SharedContentPacks=(MountName="Geometry",DetailLevels=("Standard","High"))
+SharedContentPacks=(MountName="FPWeapon",DetailLevels=("Standard"))
+SharedContentPacks=(MountName="VRSpectator",DetailLevels=("Standard"))
+EditDetailLevelPreference="High"
+FocusAsset=/Game/VRTemplate/Blueprints
+
+HiddenSettings=RayTracing
+HiddenSettings=Languages
+HiddenSettings=HardwareTarget
+HiddenSettings=GraphicsPreset
+HiddenSettings=XR
\ No newline at end of file
diff --git a/Templates/TP_VirtualRealityBP/Config/config.ini b/Templates/TP_VirtualRealityBP/Config/config.ini
new file mode 100644
--- /dev/null
+++ b/Templates/TP_VirtualRealityBP/Config/config.ini
@@ -0,0 +1,20 @@
++ActionMappings=(ActionName="TeleportRight",Key=MotionController_Right_Thumbstick,bShift=False,bCtrl=False,bAlt=False,bCmd=False)
++ActionMappings=(ActionName="GrabLeft",Key=MotionController_Left_Trigger,bShift=False,bCtrl=False,bAlt=False,bCmd=False)
++ActionMappings=(ActionName="GrabRight",Key=MotionController_Right_Trigger,bShift=False,bCtrl=False,bAlt=False,bCmd=False)
++ActionMappings=(ActionName="TeleportLeft",Key=MotionController_Left_Thumbstick,bShift=False,bCtrl=False,bAlt=False,bCmd=False)
++ActionMappings=(ActionName="HMDTeleport",Key=Gamepad_FaceButton_Bottom,bShift=False,bCtrl=False,bAlt=False,bCmd=False)
++ActionMappings=(ActionName="HMDTeleport",Key=SpaceBar,bShift=False,bCtrl=False,bAlt=False,bCmd=False)
++ActionMappings=(ActionName="HMDTeleport",Key=Android_Back,bShift=False,bCtrl=False,bAlt=False,bCmd=False)
++AxisMappings=(AxisName="TeleportDirectionUp",Key=Gamepad_LeftY,Scale=1.000000)
++AxisMappings=(AxisName="TeleportDirectionRight",Key=Gamepad_LeftX,Scale=1.000000)
++AxisMappings=(AxisName="MotionControllerThumbLeft_Y",Key=MotionController_Left_Thumbstick_Y,Scale=-1.000000)
++AxisMappings=(AxisName="MotionControllerThumbLeft_X",Key=MotionController_Left_Thumbstick_X,Scale=1.000000)
++AxisMappings=(AxisName="MotionControllerThumbRight_Y",Key=MotionController_Right_Thumbstick_Y,Scale=-1.000000)
++AxisMappings=(AxisName="MotionControllerThumbRight_X",Key=MotionController_Right_Thumbstick_X,Scale=1.000000)
+
+
+[Redirects]
+GameName=TP_VirtualRealityBP
+
+[AdditionalFilesToAdd]
++Files=Templates/TP_VirtualRealityBP/Content/*.*
\ No newline at end of file
diff --git a/Templates/TP_VirtualRealityBP/TP_VirtualRealityBP.uproject b/Templates/TP_VirtualRealityBP/TP_VirtualRealityBP.uproject
new file mode 100644
--- /dev/null
+++ b/Templates/TP_VirtualRealityBP/TP_VirtualRealityBP.uproject
@@ -0,0 +1,57 @@
+{
+	"FileVersion": 3,
+	"EngineAssociation": "",
+	"Category": "",
+	"Description": "",
+	"Plugins": [
+		{
+			"Name": "XRVisualization",
+			"Enabled": false
+		},
+		{
+			"Name": "ResonanceAudio",
+			"Enabled": true
+		},
+		{
+			"Name": "OculusVR",
+			"Enabled": false,
+			"SupportedTargetPlatforms": [
+				"Win32",
+				"Win64",
+				"Android"
+			]
+		},
+		{
+			"Name": "OpenXR",
+			"Enabled": true,
+			"SupportedTargetPlatforms": [
+				"Win32",
+				"Win64",
+				"Android",
+				"HoloLens"
+			]
+		},
+		{
+			"Name": "SteamVR",
+			"Enabled": false,
+			"SupportedTargetPlatforms": [
+				"Win32",
+				"Win64",
+				"Linux"
+			]
+		}
+	],
+	"TargetPlatforms": [
+		"Android",
+		"IOS",
+		"Linux",
+		"Mac",
+		"PS4",
+		"PS5",
+		"Windows",
+		"WindowsNoEditor",
+		"Switch",
+		"XboxOne",
+		"XSX"
+	]
+}
\ No newline at end of file
diff --git a/Templates/TP_VirtualRealityBP/contents.txt b/Templates/TP_VirtualRealityBP/contents.txt
new file mode 100644
--- /dev/null
+++ b/Templates/TP_VirtualRealityBP/contents.txt
@@ -0,0 +1,4 @@
+../../../Templates/TP_VirtualRealityBP/manifest.json
+../../../Templates/TP_VirtualRealityBP/Config/config.ini
+../../../Templates/TP_VirtualRealityBP/Media/TP_VirtualRealityBP.png
+../../../Templates/TP_VirtualRealityBP/Media/TP_VirtualRealityBP_Preview.png
diff --git a/Templates/TP_VirtualRealityBP/manifest.json b/Templates/TP_VirtualRealityBP/manifest.json
new file mode 100644
--- /dev/null
+++ b/Templates/TP_VirtualRealityBP/manifest.json
@@ -0,0 +1,68 @@
+{
+	"Version": 1,
+	"Name": [
+		{
+			"Language": "en",
+			"Text": "Virtual Reality"
+		},
+		{
+			"Language": "ko",
+			"Text": "가상 현실"
+		},
+		{
+			"Language": "ja",
+			"Text": "バーチャルリアリティ"
+		},
+		{
+			"Language": "zh-Hans",
+			"Text": "虚拟现实应用"
+		}
+	],
+	"Description": [
+		{
+			"Language": "en",
+			"Text": "Blueprint Virtual Reality Template using OpenXR for Desktop, Console, and Mobile. The template features teleport locomotion, snap turning, grabbable objects, grabbable and interactive objects, Google Resonance spatialized audio, and a VR Spectator Camera."
+		},
+		{
+			"Language": "ko",
+			"Text": "데스크톱 및 Playstation VR 용 블루프린트 가상 현실 템플릿입니다. 각 플랫폼 별 전용 폰과 맵은 물론 두 가지 텔레포트 이동 방식이 포함되어 있습니다. 플레이어는 MotionControllerMap 에서 모션 컨트롤러를 사용하여 텔레포트할 수 있으며, HMDLocomotionMap 에서 게임패드 아니면 HeadMountedDisplay 만으로 이동할 곳을 지정하는 방식으로 텔레포트할 수 있습니다."
+		},
+		{
+			"Language": "ja",
+			"Text": "デスクトップ VR と Playstation VR 向け、ブループリントのバーチャルリアリティテンプレートです。テンプレートには二種類のテレポートロコモーションモードがあります。それぞれに ポーンとマップが用意されています。MotionControllerMap ではモーションコントローラーを使ったテレポートが可能です。HMDLocomotionMap ではゲームパッドを使ったテレポート、及び、ヘッドマウントディスプレイを使ったポイントアンドテレポートが可能です。"
+		},
+		{
+			"Language": "zh-Hans",
+			"Text": "蓝图虚拟现实应用模板面向台式PC以及Playstation VR设备。模板包含2中传送移动模式，分别拥有独立的Pawn和地图。MotionControllerMap可以让玩家使用体感手柄四处传送，而HMDLocomotionMap可以让玩家使用游戏手柄或头显设备进行“指向-传送”式的传送。"
+		}
+	],
+	"AssetTypes": [
+		{
+			"Language": "en",
+			"Text": "Sound Cue, Haptic Feedback Effect Curve, Niagara System."
+		}
+	],
+	"SearchTags": [
+		{
+			"Language": "en",
+			"Text": "VR,Virtual,Reality,Virtual Reality"
+		}
+	],
+	"ClassTypes": "Pawn, GameMode, Blueprint Interface, Actor",
+	"Category": "BlueprintFeature",
+	"FocusAsset": "/Game/VRTemplate/Blueprints",
+	"Thumbnail": "TP_VirtualRealityBP.png",
+	"Screenshots": [
+		"TP_VirtualRealityBP_Preview.png"
+	],
+	"SortKey": "VirtualRealityBP",
+	"AdditionalFeaturePacks": [
+		{
+			"MountName": "Geometry",
+			"DetailLevels": [
+				"Standard",
+				"High"
+			]
+		}
+	]
+}
