Rebuilt VR Template in UE5 EA

diff --git a/Templates/TP_VirtualRealityBP/Config/DefaultEngine.ini b/Templates/TP_VirtualRealityBP/Config/DefaultEngine.ini
--- a/Templates/TP_VirtualRealityBP/Config/DefaultEngine.ini
+++ b/Templates/TP_VirtualRealityBP/Config/DefaultEngine.ini
@@ -141,8 +141,30 @@ GlobalDefaultServerGameMode=None
 bExplicitCanvasChildZOrder=True
 
 [/Script/WindowsTargetPlatform.WindowsTargetSettings]
+Compiler=Default
+MinimumOSVersion=MSOS_Vista
+bTarget32Bit=False
+AudioSampleRate=48000
+AudioCallbackBufferFrameSize=1024
+AudioNumBuffersToEnqueue=1
+AudioMaxChannels=0
+AudioNumSourceWorkers=4
 SpatializationPlugin=Resonance Audio
 ReverbPlugin=Resonance Audio
+OcclusionPlugin=
+CompressionOverrides=(bOverrideCompressionTimes=False,DurationThreshold=5.000000,MaxNumRandomBranches=0,SoundCueQualityIndex=0)
+bUseAudioStreamCaching=False
+CacheSizeKB=0
+MaxChunkSizeOverrideKB=0
+bResampleForDevice=False
+MaxSampleRate=0.000000
+HighSampleRate=0.000000
+MedSampleRate=0.000000
+LowSampleRate=0.000000
+MinSampleRate=0.000000
+CompressionQualityModifier=0.000000
+AutoStreamingThreshold=0.000000
+SoundCueCookQualityIndex=-1
 
 [/Script/AndroidRuntimeSettings.AndroidRuntimeSettings]
 PackageName=com.epicgames.[PROJECT]
@@ -261,3 +283,64 @@ SpatializationPlugin=Built-in Spatialization
 SpatializationPlugin=Resonance Audio
 ReverbPlugin=Resonance Audio
 
+[/Script/OculusHMD.OculusHMDRuntimeSettings]
+HandTrackingSupport=ControllersAndHands
+FFRLevel=FFR_Medium
+
+[/Script/Engine.CollisionProfile]
+-Profiles=(Name="NoCollision",CollisionEnabled=NoCollision,ObjectTypeName="WorldStatic",CustomResponses=((Channel="Visibility",Response=ECR_Ignore),(Channel="Camera",Response=ECR_Ignore)),HelpMessage="No collision",bCanModify=False)
+-Profiles=(Name="BlockAll",CollisionEnabled=QueryAndPhysics,ObjectTypeName="WorldStatic",CustomResponses=,HelpMessage="WorldStatic object that blocks all actors by default. All new custom channels will use its own default response. ",bCanModify=False)
+-Profiles=(Name="OverlapAll",CollisionEnabled=QueryOnly,ObjectTypeName="WorldStatic",CustomResponses=((Channel="WorldStatic",Response=ECR_Overlap),(Channel="Pawn",Response=ECR_Overlap),(Channel="Visibility",Response=ECR_Overlap),(Channel="WorldDynamic",Response=ECR_Overlap),(Channel="Camera",Response=ECR_Overlap),(Channel="PhysicsBody",Response=ECR_Overlap),(Channel="Vehicle",Response=ECR_Overlap),(Channel="Destructible",Response=ECR_Overlap)),HelpMessage="WorldStatic object that overlaps all actors by default. All new custom channels will use its own default response. ",bCanModify=False)
+-Profiles=(Name="BlockAllDynamic",CollisionEnabled=QueryAndPhysics,ObjectTypeName="WorldDynamic",CustomResponses=,HelpMessage="WorldDynamic object that blocks all actors by default. All new custom channels will use its own default response. ",bCanModify=False)
+-Profiles=(Name="OverlapAllDynamic",CollisionEnabled=QueryOnly,ObjectTypeName="WorldDynamic",CustomResponses=((Channel="WorldStatic",Response=ECR_Overlap),(Channel="Pawn",Response=ECR_Overlap),(Channel="Visibility",Response=ECR_Overlap),(Channel="WorldDynamic",Response=ECR_Overlap),(Channel="Camera",Response=ECR_Overlap),(Channel="PhysicsBody",Response=ECR_Overlap),(Channel="Vehicle",Response=ECR_Overlap),(Channel="Destructible",Response=ECR_Overlap)),HelpMessage="WorldDynamic object that overlaps all actors by default. All new custom channels will use its own default response. ",bCanModify=False)
+-Profiles=(Name="IgnoreOnlyPawn",CollisionEnabled=QueryOnly,ObjectTypeName="WorldDynamic",CustomResponses=((Channel="Pawn",Response=ECR_Ignore),(Channel="Vehicle",Response=ECR_Ignore)),HelpMessage="WorldDynamic object that ignores Pawn and Vehicle. All other channels will be set to default.",bCanModify=False)
+-Profiles=(Name="OverlapOnlyPawn",CollisionEnabled=QueryOnly,ObjectTypeName="WorldDynamic",CustomResponses=((Channel="Pawn",Response=ECR_Overlap),(Channel="Vehicle",Response=ECR_Overlap),(Channel="Camera",Response=ECR_Ignore)),HelpMessage="WorldDynamic object that overlaps Pawn, Camera, and Vehicle. All other channels will be set to default. ",bCanModify=False)
+-Profiles=(Name="Pawn",CollisionEnabled=QueryAndPhysics,ObjectTypeName="Pawn",CustomResponses=((Channel="Visibility",Response=ECR_Ignore)),HelpMessage="Pawn object. Can be used for capsule of any playerable character or AI. ",bCanModify=False)
+-Profiles=(Name="Spectator",CollisionEnabled=QueryOnly,ObjectTypeName="Pawn",CustomResponses=((Channel="WorldStatic",Response=ECR_Block),(Channel="Pawn",Response=ECR_Ignore),(Channel="Visibility",Response=ECR_Ignore),(Channel="WorldDynamic",Response=ECR_Ignore),(Channel="Camera",Response=ECR_Ignore),(Channel="PhysicsBody",Response=ECR_Ignore),(Channel="Vehicle",Response=ECR_Ignore),(Channel="Destructible",Response=ECR_Ignore)),HelpMessage="Pawn object that ignores all other actors except WorldStatic.",bCanModify=False)
+-Profiles=(Name="CharacterMesh",CollisionEnabled=QueryOnly,ObjectTypeName="Pawn",CustomResponses=((Channel="Pawn",Response=ECR_Ignore),(Channel="Vehicle",Response=ECR_Ignore),(Channel="Visibility",Response=ECR_Ignore)),HelpMessage="Pawn object that is used for Character Mesh. All other channels will be set to default.",bCanModify=False)
+-Profiles=(Name="PhysicsActor",CollisionEnabled=QueryAndPhysics,ObjectTypeName="PhysicsBody",CustomResponses=,HelpMessage="Simulating actors",bCanModify=False)
+-Profiles=(Name="Destructible",CollisionEnabled=QueryAndPhysics,ObjectTypeName="Destructible",CustomResponses=,HelpMessage="Destructible actors",bCanModify=False)
+-Profiles=(Name="InvisibleWall",CollisionEnabled=QueryAndPhysics,ObjectTypeName="WorldStatic",CustomResponses=((Channel="Visibility",Response=ECR_Ignore)),HelpMessage="WorldStatic object that is invisible.",bCanModify=False)
+-Profiles=(Name="InvisibleWallDynamic",CollisionEnabled=QueryAndPhysics,ObjectTypeName="WorldDynamic",CustomResponses=((Channel="Visibility",Response=ECR_Ignore)),HelpMessage="WorldDynamic object that is invisible.",bCanModify=False)
+-Profiles=(Name="Trigger",CollisionEnabled=QueryOnly,ObjectTypeName="WorldDynamic",CustomResponses=((Channel="WorldStatic",Response=ECR_Overlap),(Channel="Pawn",Response=ECR_Overlap),(Channel="Visibility",Response=ECR_Ignore),(Channel="WorldDynamic",Response=ECR_Overlap),(Channel="Camera",Response=ECR_Overlap),(Channel="PhysicsBody",Response=ECR_Overlap),(Channel="Vehicle",Response=ECR_Overlap),(Channel="Destructible",Response=ECR_Overlap)),HelpMessage="WorldDynamic object that is used for trigger. All other channels will be set to default.",bCanModify=False)
+-Profiles=(Name="Ragdoll",CollisionEnabled=QueryAndPhysics,ObjectTypeName="PhysicsBody",CustomResponses=((Channel="Pawn",Response=ECR_Ignore),(Channel="Visibility",Response=ECR_Ignore)),HelpMessage="Simulating Skeletal Mesh Component. All other channels will be set to default.",bCanModify=False)
+-Profiles=(Name="Vehicle",CollisionEnabled=QueryAndPhysics,ObjectTypeName="Vehicle",CustomResponses=,HelpMessage="Vehicle object that blocks Vehicle, WorldStatic, and WorldDynamic. All other channels will be set to default.",bCanModify=False)
+-Profiles=(Name="UI",CollisionEnabled=QueryOnly,ObjectTypeName="WorldDynamic",CustomResponses=((Channel="WorldStatic",Response=ECR_Overlap),(Channel="Pawn",Response=ECR_Overlap),(Channel="Visibility",Response=ECR_Block),(Channel="WorldDynamic",Response=ECR_Overlap),(Channel="Camera",Response=ECR_Overlap),(Channel="PhysicsBody",Response=ECR_Overlap),(Channel="Vehicle",Response=ECR_Overlap),(Channel="Destructible",Response=ECR_Overlap)),HelpMessage="WorldStatic object that overlaps all actors by default. All new custom channels will use its own default response. ",bCanModify=False)
++Profiles=(Name="NoCollision",CollisionEnabled=NoCollision,bCanModify=False,ObjectTypeName="WorldStatic",CustomResponses=((Channel="Visibility",Response=ECR_Ignore),(Channel="Camera",Response=ECR_Ignore)),HelpMessage="No collision")
++Profiles=(Name="BlockAll",CollisionEnabled=QueryAndPhysics,bCanModify=False,ObjectTypeName="WorldStatic",CustomResponses=,HelpMessage="WorldStatic object that blocks all actors by default. All new custom channels will use its own default response. ")
++Profiles=(Name="OverlapAll",CollisionEnabled=QueryOnly,bCanModify=False,ObjectTypeName="WorldStatic",CustomResponses=((Channel="WorldStatic",Response=ECR_Overlap),(Channel="Pawn",Response=ECR_Overlap),(Channel="Visibility",Response=ECR_Overlap),(Channel="WorldDynamic",Response=ECR_Overlap),(Channel="Camera",Response=ECR_Overlap),(Channel="PhysicsBody",Response=ECR_Overlap),(Channel="Vehicle",Response=ECR_Overlap),(Channel="Destructible",Response=ECR_Overlap)),HelpMessage="WorldStatic object that overlaps all actors by default. All new custom channels will use its own default response. ")
++Profiles=(Name="BlockAllDynamic",CollisionEnabled=QueryAndPhysics,bCanModify=False,ObjectTypeName="WorldDynamic",CustomResponses=,HelpMessage="WorldDynamic object that blocks all actors by default. All new custom channels will use its own default response. ")
++Profiles=(Name="OverlapAllDynamic",CollisionEnabled=QueryOnly,bCanModify=False,ObjectTypeName="WorldDynamic",CustomResponses=((Channel="WorldStatic",Response=ECR_Overlap),(Channel="Pawn",Response=ECR_Overlap),(Channel="Visibility",Response=ECR_Overlap),(Channel="WorldDynamic",Response=ECR_Overlap),(Channel="Camera",Response=ECR_Overlap),(Channel="PhysicsBody",Response=ECR_Overlap),(Channel="Vehicle",Response=ECR_Overlap),(Channel="Destructible",Response=ECR_Overlap)),HelpMessage="WorldDynamic object that overlaps all actors by default. All new custom channels will use its own default response. ")
++Profiles=(Name="IgnoreOnlyPawn",CollisionEnabled=QueryOnly,bCanModify=False,ObjectTypeName="WorldDynamic",CustomResponses=((Channel="Pawn",Response=ECR_Ignore),(Channel="Vehicle",Response=ECR_Ignore)),HelpMessage="WorldDynamic object that ignores Pawn and Vehicle. All other channels will be set to default.")
++Profiles=(Name="OverlapOnlyPawn",CollisionEnabled=QueryOnly,bCanModify=False,ObjectTypeName="WorldDynamic",CustomResponses=((Channel="Pawn",Response=ECR_Overlap),(Channel="Vehicle",Response=ECR_Overlap),(Channel="Camera",Response=ECR_Ignore)),HelpMessage="WorldDynamic object that overlaps Pawn, Camera, and Vehicle. All other channels will be set to default. ")
++Profiles=(Name="Pawn",CollisionEnabled=QueryAndPhysics,bCanModify=False,ObjectTypeName="Pawn",CustomResponses=((Channel="Visibility",Response=ECR_Ignore)),HelpMessage="Pawn object. Can be used for capsule of any playerable character or AI. ")
++Profiles=(Name="Spectator",CollisionEnabled=QueryOnly,bCanModify=False,ObjectTypeName="Pawn",CustomResponses=((Channel="WorldStatic"),(Channel="Pawn",Response=ECR_Ignore),(Channel="Visibility",Response=ECR_Ignore),(Channel="WorldDynamic",Response=ECR_Ignore),(Channel="Camera",Response=ECR_Ignore),(Channel="PhysicsBody",Response=ECR_Ignore),(Channel="Vehicle",Response=ECR_Ignore),(Channel="Destructible",Response=ECR_Ignore)),HelpMessage="Pawn object that ignores all other actors except WorldStatic.")
++Profiles=(Name="CharacterMesh",CollisionEnabled=QueryOnly,bCanModify=False,ObjectTypeName="Pawn",CustomResponses=((Channel="Pawn",Response=ECR_Ignore),(Channel="Vehicle",Response=ECR_Ignore),(Channel="Visibility",Response=ECR_Ignore)),HelpMessage="Pawn object that is used for Character Mesh. All other channels will be set to default.")
++Profiles=(Name="PhysicsActor",CollisionEnabled=QueryAndPhysics,bCanModify=False,ObjectTypeName="PhysicsBody",CustomResponses=,HelpMessage="Simulating actors")
++Profiles=(Name="Destructible",CollisionEnabled=QueryAndPhysics,bCanModify=False,ObjectTypeName="Destructible",CustomResponses=,HelpMessage="Destructible actors")
++Profiles=(Name="InvisibleWall",CollisionEnabled=QueryAndPhysics,bCanModify=False,ObjectTypeName="WorldStatic",CustomResponses=((Channel="Visibility",Response=ECR_Ignore)),HelpMessage="WorldStatic object that is invisible.")
++Profiles=(Name="InvisibleWallDynamic",CollisionEnabled=QueryAndPhysics,bCanModify=False,ObjectTypeName="WorldDynamic",CustomResponses=((Channel="Visibility",Response=ECR_Ignore)),HelpMessage="WorldDynamic object that is invisible.")
++Profiles=(Name="Trigger",CollisionEnabled=QueryOnly,bCanModify=False,ObjectTypeName="WorldDynamic",CustomResponses=((Channel="WorldStatic",Response=ECR_Overlap),(Channel="Pawn",Response=ECR_Overlap),(Channel="Visibility",Response=ECR_Ignore),(Channel="WorldDynamic",Response=ECR_Overlap),(Channel="Camera",Response=ECR_Overlap),(Channel="PhysicsBody",Response=ECR_Overlap),(Channel="Vehicle",Response=ECR_Overlap),(Channel="Destructible",Response=ECR_Overlap)),HelpMessage="WorldDynamic object that is used for trigger. All other channels will be set to default.")
++Profiles=(Name="Ragdoll",CollisionEnabled=QueryAndPhysics,bCanModify=False,ObjectTypeName="PhysicsBody",CustomResponses=((Channel="Pawn",Response=ECR_Ignore),(Channel="Visibility",Response=ECR_Ignore)),HelpMessage="Simulating Skeletal Mesh Component. All other channels will be set to default.")
++Profiles=(Name="Vehicle",CollisionEnabled=QueryAndPhysics,bCanModify=False,ObjectTypeName="Vehicle",CustomResponses=,HelpMessage="Vehicle object that blocks Vehicle, WorldStatic, and WorldDynamic. All other channels will be set to default.")
++Profiles=(Name="UI",CollisionEnabled=QueryOnly,bCanModify=False,ObjectTypeName="WorldDynamic",CustomResponses=((Channel="WorldStatic",Response=ECR_Overlap),(Channel="Pawn",Response=ECR_Overlap),(Channel="Visibility"),(Channel="WorldDynamic",Response=ECR_Overlap),(Channel="Camera",Response=ECR_Overlap),(Channel="PhysicsBody",Response=ECR_Overlap),(Channel="Vehicle",Response=ECR_Overlap),(Channel="Destructible",Response=ECR_Overlap)),HelpMessage="WorldStatic object that overlaps all actors by default. All new custom channels will use its own default response. ")
++DefaultChannelResponses=(Channel=ECC_GameTraceChannel1,DefaultResponse=ECR_Ignore,bTraceType=True,bStaticObject=False,Name="3DWidget")
+-ProfileRedirects=(OldName="BlockingVolume",NewName="InvisibleWall")
+-ProfileRedirects=(OldName="InterpActor",NewName="IgnoreOnlyPawn")
+-ProfileRedirects=(OldName="StaticMeshComponent",NewName="BlockAllDynamic")
+-ProfileRedirects=(OldName="SkeletalMeshActor",NewName="PhysicsActor")
+-ProfileRedirects=(OldName="InvisibleActor",NewName="InvisibleWallDynamic")
++ProfileRedirects=(OldName="BlockingVolume",NewName="InvisibleWall")
++ProfileRedirects=(OldName="InterpActor",NewName="IgnoreOnlyPawn")
++ProfileRedirects=(OldName="StaticMeshComponent",NewName="BlockAllDynamic")
++ProfileRedirects=(OldName="SkeletalMeshActor",NewName="PhysicsActor")
++ProfileRedirects=(OldName="InvisibleActor",NewName="InvisibleWallDynamic")
+-CollisionChannelRedirects=(OldName="Static",NewName="WorldStatic")
+-CollisionChannelRedirects=(OldName="Dynamic",NewName="WorldDynamic")
+-CollisionChannelRedirects=(OldName="VehicleMovement",NewName="Vehicle")
+-CollisionChannelRedirects=(OldName="PawnMovement",NewName="Pawn")
++CollisionChannelRedirects=(OldName="Static",NewName="WorldStatic")
++CollisionChannelRedirects=(OldName="Dynamic",NewName="WorldDynamic")
++CollisionChannelRedirects=(OldName="VehicleMovement",NewName="Vehicle")
++CollisionChannelRedirects=(OldName="PawnMovement",NewName="Pawn")
+
diff --git a/Templates/TP_VirtualRealityBP/Config/DefaultInput.ini b/Templates/TP_VirtualRealityBP/Config/DefaultInput.ini
--- a/Templates/TP_VirtualRealityBP/Config/DefaultInput.ini
+++ b/Templates/TP_VirtualRealityBP/Config/DefaultInput.ini
@@ -79,6 +79,28 @@
 +AxisConfig=(AxisKeyName="MagicLeap_Right_Touch1_X",AxisProperties=(DeadZone=0.000000,Sensitivity=1.000000,Exponent=1.000000,bInvert=False))
 +AxisConfig=(AxisKeyName="MagicLeap_Right_Touch1_Y",AxisProperties=(DeadZone=0.000000,Sensitivity=1.000000,Exponent=1.000000,bInvert=False))
 +AxisConfig=(AxisKeyName="MagicLeap_Right_Touch1_Force",AxisProperties=(DeadZone=0.000000,Sensitivity=1.000000,Exponent=1.000000,bInvert=False))
++AxisConfig=(AxisKeyName="OculusTouch_Left_Thumbstick",AxisProperties=(DeadZone=0.000000,Sensitivity=1.000000,Exponent=1.000000,bInvert=False))
++AxisConfig=(AxisKeyName="OculusTouch_Left_FaceButton1",AxisProperties=(DeadZone=0.000000,Sensitivity=1.000000,Exponent=1.000000,bInvert=False))
++AxisConfig=(AxisKeyName="OculusTouch_Left_Trigger",AxisProperties=(DeadZone=0.000000,Sensitivity=1.000000,Exponent=1.000000,bInvert=False))
++AxisConfig=(AxisKeyName="OculusTouch_Left_FaceButton2",AxisProperties=(DeadZone=0.000000,Sensitivity=1.000000,Exponent=1.000000,bInvert=False))
++AxisConfig=(AxisKeyName="OculusTouch_Left_IndexPointing",AxisProperties=(DeadZone=0.000000,Sensitivity=1.000000,Exponent=1.000000,bInvert=False))
++AxisConfig=(AxisKeyName="OculusTouch_Left_ThumbUp",AxisProperties=(DeadZone=0.000000,Sensitivity=1.000000,Exponent=1.000000,bInvert=False))
++AxisConfig=(AxisKeyName="OculusTouch_Right_Thumbstick",AxisProperties=(DeadZone=0.000000,Sensitivity=1.000000,Exponent=1.000000,bInvert=False))
++AxisConfig=(AxisKeyName="OculusTouch_Right_FaceButton1",AxisProperties=(DeadZone=0.000000,Sensitivity=1.000000,Exponent=1.000000,bInvert=False))
++AxisConfig=(AxisKeyName="OculusTouch_Right_Trigger",AxisProperties=(DeadZone=0.000000,Sensitivity=1.000000,Exponent=1.000000,bInvert=False))
++AxisConfig=(AxisKeyName="OculusTouch_Right_FaceButton2",AxisProperties=(DeadZone=0.000000,Sensitivity=1.000000,Exponent=1.000000,bInvert=False))
++AxisConfig=(AxisKeyName="OculusTouch_Right_IndexPointing",AxisProperties=(DeadZone=0.000000,Sensitivity=1.000000,Exponent=1.000000,bInvert=False))
++AxisConfig=(AxisKeyName="OculusTouch_Right_ThumbUp",AxisProperties=(DeadZone=0.000000,Sensitivity=1.000000,Exponent=1.000000,bInvert=False))
++AxisConfig=(AxisKeyName="OculusHand_Left_ThumbPinchStrength",AxisProperties=(DeadZone=0.000000,Sensitivity=1.000000,Exponent=1.000000,bInvert=False))
++AxisConfig=(AxisKeyName="OculusHand_Left_IndexPinchStrength",AxisProperties=(DeadZone=0.000000,Sensitivity=1.000000,Exponent=1.000000,bInvert=False))
++AxisConfig=(AxisKeyName="OculusHand_Left_MiddlePinchStrength",AxisProperties=(DeadZone=0.000000,Sensitivity=1.000000,Exponent=1.000000,bInvert=False))
++AxisConfig=(AxisKeyName="OculusHand_Left_RingPinchStrength",AxisProperties=(DeadZone=0.000000,Sensitivity=1.000000,Exponent=1.000000,bInvert=False))
++AxisConfig=(AxisKeyName="OculusHand_Left_PinkPinchStrength",AxisProperties=(DeadZone=0.000000,Sensitivity=1.000000,Exponent=1.000000,bInvert=False))
++AxisConfig=(AxisKeyName="OculusHand_Right_ThumbPinchStrength",AxisProperties=(DeadZone=0.000000,Sensitivity=1.000000,Exponent=1.000000,bInvert=False))
++AxisConfig=(AxisKeyName="OculusHand_Right_IndexPinchStrength",AxisProperties=(DeadZone=0.000000,Sensitivity=1.000000,Exponent=1.000000,bInvert=False))
++AxisConfig=(AxisKeyName="OculusHand_Right_MiddlePinchStrength",AxisProperties=(DeadZone=0.000000,Sensitivity=1.000000,Exponent=1.000000,bInvert=False))
++AxisConfig=(AxisKeyName="OculusHand_Right_RingPinchStrength",AxisProperties=(DeadZone=0.000000,Sensitivity=1.000000,Exponent=1.000000,bInvert=False))
++AxisConfig=(AxisKeyName="OculusHand_Right_PinkPinchStrength",AxisProperties=(DeadZone=0.000000,Sensitivity=1.000000,Exponent=1.000000,bInvert=False))
 bAltEnterTogglesFullscreen=True
 bF11TogglesFullscreen=True
 bUseMouseForTouch=False
@@ -113,6 +135,13 @@ DoubleClickTime=0.200000
 +ActionMappings=(ActionName="SpectatorFadeToggle",bShift=False,bCtrl=False,bAlt=False,bCmd=False,Key=F)
 +ActionMappings=(ActionName="SpectatorActiveToggle",bShift=False,bCtrl=False,bAlt=False,bCmd=False,Key=Tab)
 +ActionMappings=(ActionName="SpectatorResetRotation",bShift=False,bCtrl=False,bAlt=False,bCmd=False,Key=R)
++ActionMappings=(ActionName="MenuToggleLeft",bShift=False,bCtrl=False,bAlt=False,bCmd=False,Key=OculusTouch_Left_Menu_Click)
++ActionMappings=(ActionName="MenuToggleLeft",bShift=False,bCtrl=False,bAlt=False,bCmd=False,Key=ValveIndex_Left_System_Click)
++ActionMappings=(ActionName="MenuToggleLeft",bShift=False,bCtrl=False,bAlt=False,bCmd=False,Key=MixedReality_Left_Menu_Click)
++ActionMappings=(ActionName="MenuToggleLeft",bShift=False,bCtrl=False,bAlt=False,bCmd=False,Key=Vive_Left_Menu_Click)
++ActionMappings=(ActionName="MenuToggleRight",bShift=False,bCtrl=False,bAlt=False,bCmd=False,Key=ValveIndex_Right_System_Click)
++ActionMappings=(ActionName="MenuToggleRight",bShift=False,bCtrl=False,bAlt=False,bCmd=False,Key=MixedReality_Right_Menu_Click)
++ActionMappings=(ActionName="MenuToggleRight",bShift=False,bCtrl=False,bAlt=False,bCmd=False,Key=Vive_Right_Menu_Click)
 +AxisMappings=(AxisName="MovementAxisLeft_X",Scale=1.000000,Key=OculusTouch_Left_Thumbstick_X)
 +AxisMappings=(AxisName="GrabAxisLeft",Scale=1.000000,Key=OculusTouch_Left_Grip_Axis)
 +AxisMappings=(AxisName="MovementAxisLeft_Y",Scale=1.000000,Key=OculusTouch_Left_Thumbstick_Y)
diff --git a/Templates/TP_VirtualRealityBP/TP_VirtualRealityBP.uproject b/Templates/TP_VirtualRealityBP/TP_VirtualRealityBP.uproject
--- a/Templates/TP_VirtualRealityBP/TP_VirtualRealityBP.uproject
+++ b/Templates/TP_VirtualRealityBP/TP_VirtualRealityBP.uproject
@@ -12,14 +12,6 @@
 			"Name": "ResonanceAudio",
 			"Enabled": true
 		},
-		{
-			"Name": "OculusVR",
-			"Enabled": false,
-			"SupportedTargetPlatforms": [
-				"Win64",
-				"Android"
-			]
-		},
 		{
 			"Name": "OpenXR",
 			"Enabled": true,
@@ -28,14 +20,6 @@
 				"Android",
 				"HoloLens"
 			]
-		},
-		{
-			"Name": "SteamVR",
-			"Enabled": false,
-			"SupportedTargetPlatforms": [
-				"Win64",
-				"Linux"
-			]
 		}
 	],
 	"TargetPlatforms": [
