Removes r.UsePreExposure in project templates

diff --git a/Templates/TP_HandheldARBP/Config/DefaultEngine.ini b/Templates/TP_HandheldARBP/Config/DefaultEngine.ini
--- a/Templates/TP_HandheldARBP/Config/DefaultEngine.ini
+++ b/Templates/TP_HandheldARBP/Config/DefaultEngine.ini
@@ -30,7 +30,6 @@ TranslucentSortAxis=(X=0.000000,Y=-1.000000,Z=0.000000)
 r.CustomDepth=1
 r.CustomDepthTemporalAAJitter=True
 r.PostProcessing.PropagateAlpha=0
-r.UsePreExposure=False
 r.DefaultFeature.Bloom=True
 r.DefaultFeature.AmbientOcclusion=False
 r.DefaultFeature.AmbientOcclusionStaticFraction=True
diff --git a/Templates/TP_VirtualRealityBP/Config/DefaultEngine.ini b/Templates/TP_VirtualRealityBP/Config/DefaultEngine.ini
--- a/Templates/TP_VirtualRealityBP/Config/DefaultEngine.ini
+++ b/Templates/TP_VirtualRealityBP/Config/DefaultEngine.ini
@@ -48,7 +48,6 @@ r.DefaultFeature.AutoExposure=False
 r.DefaultFeature.AutoExposure.Method=0
 r.DefaultFeature.AutoExposure.Bias=1.000000
 r.DefaultFeature.AutoExposure.ExtendDefaultLuminanceRange=True
-r.UsePreExposure=True
 r.EyeAdaptation.EditorOnly=False
 r.DefaultFeature.MotionBlur=False
 r.DefaultFeature.LensFlare=False
