Remove r.MobileMsaa from TP ini

diff --git a/Templates/TP_HandheldARBP/Config/DefaultEngine.ini b/Templates/TP_HandheldARBP/Config/DefaultEngine.ini
--- a/Templates/TP_HandheldARBP/Config/DefaultEngine.ini
+++ b/Templates/TP_HandheldARBP/Config/DefaultEngine.ini
@@ -3,7 +3,6 @@
 r.MobileHDR=True
 r.Mobile.DisableVertexFog=True
 r.Shadow.CSM.MaxMobileCascades=2
-r.MobileMSAA=1
 r.DiscardUnusedQuality=False
 r.AllowOcclusionQueries=True
 r.MinScreenRadiusForLights=0.030000
diff --git a/Templates/TP_VirtualRealityBP/Config/DefaultEngine.ini b/Templates/TP_VirtualRealityBP/Config/DefaultEngine.ini
--- a/Templates/TP_VirtualRealityBP/Config/DefaultEngine.ini
+++ b/Templates/TP_VirtualRealityBP/Config/DefaultEngine.ini
@@ -3,7 +3,6 @@
 [/Script/Engine.RendererSettings]
 r.Mobile.DisableVertexFog=True
 r.Shadow.CSM.MaxMobileCascades=2
-r.MobileMSAA=4
 r.Mobile.AllowDitheredLODTransition=False
 r.Mobile.AllowSoftwareOcclusion=False
 r.Mobile.VirtualTextures=False
