XR | VR Template | Mobile Anti-Aliasing is now set to MSAA

diff --git a/Templates/TP_VirtualRealityBP/Config/DefaultEngine.ini b/Templates/TP_VirtualRealityBP/Config/DefaultEngine.ini
--- a/Templates/TP_VirtualRealityBP/Config/DefaultEngine.ini
+++ b/Templates/TP_VirtualRealityBP/Config/DefaultEngine.ini
@@ -113,6 +113,7 @@ VisualizeCalibrationColorMaterialPath=None
 VisualizeCalibrationCustomMaterialPath=None
 VisualizeCalibrationGrayscaleMaterialPath=None
 r.Mobile.SupportGPUScene=1
+r.Mobile.AntiAliasing=3
 
 [/Script/HardwareTargeting.HardwareTargetingSettings]
 TargetedHardwareClass=Mobile
