XR | VR Template
Added Gamepad Input for VRSpectator. Fixed Valve Index controllers not having input assigned to open the menu.

diff --git a/Templates/TP_VirtualRealityBP/Config/DefaultInput.ini b/Templates/TP_VirtualRealityBP/Config/DefaultInput.ini
--- a/Templates/TP_VirtualRealityBP/Config/DefaultInput.ini
+++ b/Templates/TP_VirtualRealityBP/Config/DefaultInput.ini
@@ -87,6 +87,7 @@ bUseMouseForTouch=False
 bEnableMouseSmoothing=True
 bEnableFOVScaling=True
 bCaptureMouseOnLaunch=True
+bEnableLegacyInputScales=True
 bAlwaysShowTouchInterface=False
 bShowConsoleOnFourFingerTap=True
 bEnableGestureRecognizer=False
@@ -116,12 +117,16 @@ DoubleClickTime=0.200000
 +ActionMappings=(ActionName="SpectatorActiveToggle",bShift=False,bCtrl=False,bAlt=False,bCmd=False,Key=Tab)
 +ActionMappings=(ActionName="SpectatorResetRotation",bShift=False,bCtrl=False,bAlt=False,bCmd=False,Key=R)
 +ActionMappings=(ActionName="MenuToggleLeft",bShift=False,bCtrl=False,bAlt=False,bCmd=False,Key=OculusTouch_Left_Menu_Click)
-+ActionMappings=(ActionName="MenuToggleLeft",bShift=False,bCtrl=False,bAlt=False,bCmd=False,Key=ValveIndex_Left_System_Click)
++ActionMappings=(ActionName="MenuToggleLeft",bShift=False,bCtrl=False,bAlt=False,bCmd=False,Key=ValveIndex_Left_Thumbstick_Click)
 +ActionMappings=(ActionName="MenuToggleLeft",bShift=False,bCtrl=False,bAlt=False,bCmd=False,Key=MixedReality_Left_Menu_Click)
 +ActionMappings=(ActionName="MenuToggleLeft",bShift=False,bCtrl=False,bAlt=False,bCmd=False,Key=Vive_Left_Menu_Click)
-+ActionMappings=(ActionName="MenuToggleRight",bShift=False,bCtrl=False,bAlt=False,bCmd=False,Key=ValveIndex_Right_System_Click)
++ActionMappings=(ActionName="MenuToggleRight",bShift=False,bCtrl=False,bAlt=False,bCmd=False,Key=ValveIndex_Right_Thumbstick_Click)
 +ActionMappings=(ActionName="MenuToggleRight",bShift=False,bCtrl=False,bAlt=False,bCmd=False,Key=MixedReality_Right_Menu_Click)
 +ActionMappings=(ActionName="MenuToggleRight",bShift=False,bCtrl=False,bAlt=False,bCmd=False,Key=Vive_Right_Menu_Click)
++ActionMappings=(ActionName="MenuToggleLeft",bShift=False,bCtrl=False,bAlt=False,bCmd=False,Key=OculusTouch_Left_Thumbstick_Click)
++ActionMappings=(ActionName="MenuToggleRight",bShift=False,bCtrl=False,bAlt=False,bCmd=False,Key=OculusTouch_Right_Thumbstick_Click)
++ActionMappings=(ActionName="SpectatorFadeToggle",bShift=False,bCtrl=False,bAlt=False,bCmd=False,Key=Gamepad_FaceButton_Top)
++ActionMappings=(ActionName="SpectatorActiveToggle",bShift=False,bCtrl=False,bAlt=False,bCmd=False,Key=Gamepad_FaceButton_Bottom)
 +AxisMappings=(AxisName="MovementAxisLeft_X",Scale=1.000000,Key=OculusTouch_Left_Thumbstick_X)
 +AxisMappings=(AxisName="GrabAxisLeft",Scale=1.000000,Key=OculusTouch_Left_Grip_Axis)
 +AxisMappings=(AxisName="MovementAxisLeft_Y",Scale=1.000000,Key=OculusTouch_Left_Thumbstick_Y)
@@ -165,6 +170,14 @@ DoubleClickTime=0.200000
 +AxisMappings=(AxisName="MovementAxisRight_Y",Scale=1.000000,Key=MixedReality_Right_Thumbstick_Y)
 +AxisMappings=(AxisName="MovementAxisRight_Y",Scale=1.000000,Key=Vive_Right_Trackpad_Y)
 +AxisMappings=(AxisName="MovementAxisRight_Y",Scale=1.000000,Key=OculusTouch_Right_A_Click)
++AxisMappings=(AxisName="SpectatorYaw",Scale=1.000000,Key=Gamepad_RightX)
++AxisMappings=(AxisName="SpectatorPitch",Scale=-1.000000,Key=Gamepad_RightY)
++AxisMappings=(AxisName="SpectatorMoveForward",Scale=1.000000,Key=Gamepad_LeftY)
++AxisMappings=(AxisName="SpectatorMoveRight",Scale=1.000000,Key=Gamepad_LeftX)
++AxisMappings=(AxisName="SpectatorMoveUp",Scale=1.000000,Key=Gamepad_RightShoulder)
++AxisMappings=(AxisName="SpectatorMoveUp",Scale=-1.000000,Key=Gamepad_LeftShoulder)
++AxisMappings=(AxisName="SpectatorFOV",Scale=1.000000,Key=Gamepad_RightTriggerAxis)
++AxisMappings=(AxisName="SpectatorFOV",Scale=-1.000000,Key=Gamepad_LeftTriggerAxis)
 DefaultPlayerInputClass=/Script/Engine.PlayerInput
 DefaultInputComponentClass=/Script/Engine.InputComponent
 DefaultTouchInterface=None
