Updates all project settings of project not using TSR by default

diff --git a/Templates/TP_HandheldARBP/Config/DefaultEditor.ini b/Templates/TP_HandheldARBP/Config/DefaultEditor.ini
--- a/Templates/TP_HandheldARBP/Config/DefaultEditor.ini
+++ b/Templates/TP_HandheldARBP/Config/DefaultEditor.ini
@@ -0,0 +1,3 @@
+
+[/Script/UnrealEd.EditorPerformanceProjectSettings]
+RealtimeScreenPercentageMode=BasedOnDPIScale
diff --git a/Templates/TP_VirtualRealityBP/Config/DefaultEditor.ini b/Templates/TP_VirtualRealityBP/Config/DefaultEditor.ini
--- a/Templates/TP_VirtualRealityBP/Config/DefaultEditor.ini
+++ b/Templates/TP_VirtualRealityBP/Config/DefaultEditor.ini
@@ -1 +1,3 @@
 
+[/Script/UnrealEd.EditorPerformanceProjectSettings]
+RealtimeScreenPercentageMode=BasedOnDPIScale
