Fixed VR Template not adding Input Mappings when added as a BlueprintFeature to a new project. Fixed VR Template not including VRSpectator when added as a BlueprintFeature to a new project.

diff --git a/Templates/TP_VirtualRealityBP/Config/config.ini b/Templates/TP_VirtualRealityBP/Config/config.ini
--- a/Templates/TP_VirtualRealityBP/Config/config.ini
+++ b/Templates/TP_VirtualRealityBP/Config/config.ini
@@ -1,17 +1,86 @@
-+ActionMappings=(ActionName="TeleportRight",Key=MotionController_Right_Thumbstick,bShift=False,bCtrl=False,bAlt=False,bCmd=False)
-+ActionMappings=(ActionName="GrabLeft",Key=MotionController_Left_Trigger,bShift=False,bCtrl=False,bAlt=False,bCmd=False)
-+ActionMappings=(ActionName="GrabRight",Key=MotionController_Right_Trigger,bShift=False,bCtrl=False,bAlt=False,bCmd=False)
-+ActionMappings=(ActionName="TeleportLeft",Key=MotionController_Left_Thumbstick,bShift=False,bCtrl=False,bAlt=False,bCmd=False)
-+ActionMappings=(ActionName="HMDTeleport",Key=Gamepad_FaceButton_Bottom,bShift=False,bCtrl=False,bAlt=False,bCmd=False)
-+ActionMappings=(ActionName="HMDTeleport",Key=SpaceBar,bShift=False,bCtrl=False,bAlt=False,bCmd=False)
-+ActionMappings=(ActionName="HMDTeleport",Key=Android_Back,bShift=False,bCtrl=False,bAlt=False,bCmd=False)
-+AxisMappings=(AxisName="TeleportDirectionUp",Key=Gamepad_LeftY,Scale=1.000000)
-+AxisMappings=(AxisName="TeleportDirectionRight",Key=Gamepad_LeftX,Scale=1.000000)
-+AxisMappings=(AxisName="MotionControllerThumbLeft_Y",Key=MotionController_Left_Thumbstick_Y,Scale=-1.000000)
-+AxisMappings=(AxisName="MotionControllerThumbLeft_X",Key=MotionController_Left_Thumbstick_X,Scale=1.000000)
-+AxisMappings=(AxisName="MotionControllerThumbRight_Y",Key=MotionController_Right_Thumbstick_Y,Scale=-1.000000)
-+AxisMappings=(AxisName="MotionControllerThumbRight_X",Key=MotionController_Right_Thumbstick_X,Scale=1.000000)
-
+[InputSettings]
++ActionMappings=(ActionName="GrabLeft",bShift=False,bCtrl=False,bAlt=False,bCmd=False,Key=OculusTouch_Left_Grip_Click)
++ActionMappings=(ActionName="GrabRight",bShift=False,bCtrl=False,bAlt=False,bCmd=False,Key=OculusTouch_Right_Grip_Click)
++ActionMappings=(ActionName="TriggerLeft",bShift=False,bCtrl=False,bAlt=False,bCmd=False,Key=OculusTouch_Left_Trigger_Click)
++ActionMappings=(ActionName="TriggerRight",bShift=False,bCtrl=False,bAlt=False,bCmd=False,Key=OculusTouch_Right_Trigger_Click)
++ActionMappings=(ActionName="GrabLeft",bShift=False,bCtrl=False,bAlt=False,bCmd=False,Key=ValveIndex_Left_Grip_Axis)
++ActionMappings=(ActionName="GrabLeft",bShift=False,bCtrl=False,bAlt=False,bCmd=False,Key=MixedReality_Left_Grip_Click)
++ActionMappings=(ActionName="GrabLeft",bShift=False,bCtrl=False,bAlt=False,bCmd=False,Key=Vive_Left_Grip_Click)
++ActionMappings=(ActionName="GrabRight",bShift=False,bCtrl=False,bAlt=False,bCmd=False,Key=ValveIndex_Right_Grip_Axis)
++ActionMappings=(ActionName="GrabRight",bShift=False,bCtrl=False,bAlt=False,bCmd=False,Key=MixedReality_Right_Grip_Click)
++ActionMappings=(ActionName="GrabRight",bShift=False,bCtrl=False,bAlt=False,bCmd=False,Key=Vive_Right_Grip_Click)
++ActionMappings=(ActionName="TriggerLeft",bShift=False,bCtrl=False,bAlt=False,bCmd=False,Key=ValveIndex_Left_Trigger_Click)
++ActionMappings=(ActionName="TriggerLeft",bShift=False,bCtrl=False,bAlt=False,bCmd=False,Key=MixedReality_Left_Trigger_Click)
++ActionMappings=(ActionName="TriggerLeft",bShift=False,bCtrl=False,bAlt=False,bCmd=False,Key=Vive_Left_Trigger_Click)
++ActionMappings=(ActionName="TriggerRight",bShift=False,bCtrl=False,bAlt=False,bCmd=False,Key=ValveIndex_Right_Trigger_Click)
++ActionMappings=(ActionName="TriggerRight",bShift=False,bCtrl=False,bAlt=False,bCmd=False,Key=MixedReality_Right_Trigger_Click)
++ActionMappings=(ActionName="TriggerRight",bShift=False,bCtrl=False,bAlt=False,bCmd=False,Key=Vive_Right_Trigger_Click)
++ActionMappings=(ActionName="SpectatorResetFOV",bShift=False,bCtrl=False,bAlt=False,bCmd=False,Key=MiddleMouseButton)
++ActionMappings=(ActionName="SpectatorFadeToggle",bShift=False,bCtrl=False,bAlt=False,bCmd=False,Key=F)
++ActionMappings=(ActionName="SpectatorActiveToggle",bShift=False,bCtrl=False,bAlt=False,bCmd=False,Key=Tab)
++ActionMappings=(ActionName="SpectatorResetRotation",bShift=False,bCtrl=False,bAlt=False,bCmd=False,Key=R)
++ActionMappings=(ActionName="MenuToggleLeft",bShift=False,bCtrl=False,bAlt=False,bCmd=False,Key=OculusTouch_Left_Menu_Click)
++ActionMappings=(ActionName="MenuToggleLeft",bShift=False,bCtrl=False,bAlt=False,bCmd=False,Key=ValveIndex_Left_Thumbstick_Click)
++ActionMappings=(ActionName="MenuToggleLeft",bShift=False,bCtrl=False,bAlt=False,bCmd=False,Key=MixedReality_Left_Menu_Click)
++ActionMappings=(ActionName="MenuToggleLeft",bShift=False,bCtrl=False,bAlt=False,bCmd=False,Key=Vive_Left_Menu_Click)
++ActionMappings=(ActionName="MenuToggleRight",bShift=False,bCtrl=False,bAlt=False,bCmd=False,Key=ValveIndex_Right_Thumbstick_Click)
++ActionMappings=(ActionName="MenuToggleRight",bShift=False,bCtrl=False,bAlt=False,bCmd=False,Key=MixedReality_Right_Menu_Click)
++ActionMappings=(ActionName="MenuToggleRight",bShift=False,bCtrl=False,bAlt=False,bCmd=False,Key=Vive_Right_Menu_Click)
++ActionMappings=(ActionName="MenuToggleLeft",bShift=False,bCtrl=False,bAlt=False,bCmd=False,Key=OculusTouch_Left_Thumbstick_Click)
++ActionMappings=(ActionName="MenuToggleRight",bShift=False,bCtrl=False,bAlt=False,bCmd=False,Key=OculusTouch_Right_Thumbstick_Click)
++ActionMappings=(ActionName="SpectatorFadeToggle",bShift=False,bCtrl=False,bAlt=False,bCmd=False,Key=Gamepad_FaceButton_Top)
++ActionMappings=(ActionName="SpectatorActiveToggle",bShift=False,bCtrl=False,bAlt=False,bCmd=False,Key=Gamepad_FaceButton_Bottom)
++AxisMappings=(AxisName="MovementAxisLeft_X",Scale=1.000000,Key=OculusTouch_Left_Thumbstick_X)
++AxisMappings=(AxisName="GrabAxisLeft",Scale=1.000000,Key=OculusTouch_Left_Grip_Axis)
++AxisMappings=(AxisName="MovementAxisLeft_Y",Scale=1.000000,Key=OculusTouch_Left_Thumbstick_Y)
++AxisMappings=(AxisName="TriggerAxisLeft",Scale=1.000000,Key=OculusTouch_Left_Trigger_Axis)
++AxisMappings=(AxisName="TriggerAxisRight",Scale=1.000000,Key=OculusTouch_Right_Trigger_Axis)
++AxisMappings=(AxisName="MovementAxisLeft_X",Scale=1.000000,Key=ValveIndex_Left_Thumbstick_X)
++AxisMappings=(AxisName="MovementAxisLeft_X",Scale=1.000000,Key=MixedReality_Left_Thumbstick_X)
++AxisMappings=(AxisName="MovementAxisLeft_X",Scale=1.000000,Key=Vive_Left_Trackpad_X)
++AxisMappings=(AxisName="GrabAxisLeft",Scale=1.000000,Key=ValveIndex_Left_Grip_Axis)
++AxisMappings=(AxisName="GrabAxisLeft",Scale=1.000000,Key=MixedReality_Left_Grip_Click)
++AxisMappings=(AxisName="GrabAxisLeft",Scale=1.000000,Key=Vive_Left_Grip_Click)
++AxisMappings=(AxisName="GrabAxisRight",Scale=1.000000,Key=OculusTouch_Right_Grip_Axis)
++AxisMappings=(AxisName="GrabAxisRight",Scale=1.000000,Key=ValveIndex_Right_Grip_Axis)
++AxisMappings=(AxisName="GrabAxisRight",Scale=1.000000,Key=MixedReality_Right_Grip_Click)
++AxisMappings=(AxisName="GrabAxisRight",Scale=1.000000,Key=Vive_Right_Grip_Click)
++AxisMappings=(AxisName="TriggerAxisLeft",Scale=1.000000,Key=ValveIndex_Left_Trigger_Axis)
++AxisMappings=(AxisName="TriggerAxisLeft",Scale=1.000000,Key=MixedReality_Left_Trigger_Axis)
++AxisMappings=(AxisName="TriggerAxisLeft",Scale=1.000000,Key=Vive_Left_Trigger_Axis)
++AxisMappings=(AxisName="TriggerAxisRight",Scale=1.000000,Key=ValveIndex_Right_Trigger_Axis)
++AxisMappings=(AxisName="TriggerAxisRight",Scale=1.000000,Key=MixedReality_Right_Trigger_Axis)
++AxisMappings=(AxisName="TriggerAxisRight",Scale=1.000000,Key=Vive_Right_Trigger_Axis)
++AxisMappings=(AxisName="MovementAxisLeft_Y",Scale=1.000000,Key=ValveIndex_Left_Thumbstick_Y)
++AxisMappings=(AxisName="MovementAxisLeft_Y",Scale=1.000000,Key=MixedReality_Left_Thumbstick_Y)
++AxisMappings=(AxisName="MovementAxisLeft_Y",Scale=1.000000,Key=Vive_Left_Trackpad_Y)
++AxisMappings=(AxisName="SpectatorYaw",Scale=1.000000,Key=MouseX)
++AxisMappings=(AxisName="SpectatorPitch",Scale=1.000000,Key=MouseY)
++AxisMappings=(AxisName="SpectatorFOV",Scale=1.000000,Key=MouseWheelAxis)
++AxisMappings=(AxisName="SpectatorMoveForward",Scale=1.000000,Key=W)
++AxisMappings=(AxisName="SpectatorMoveRight",Scale=1.000000,Key=D)
++AxisMappings=(AxisName="SpectatorMoveForward",Scale=-1.000000,Key=S)
++AxisMappings=(AxisName="SpectatorMoveRight",Scale=-1.000000,Key=A)
++AxisMappings=(AxisName="SpectatorMoveUp",Scale=1.000000,Key=SpaceBar)
++AxisMappings=(AxisName="SpectatorMoveUp",Scale=1.000000,Key=E)
++AxisMappings=(AxisName="SpectatorMoveUp",Scale=-1.000000,Key=Q)
++AxisMappings=(AxisName="MovementAxisRight_X",Scale=1.000000,Key=OculusTouch_Right_Thumbstick_X)
++AxisMappings=(AxisName="MovementAxisRight_X",Scale=1.000000,Key=ValveIndex_Right_Thumbstick_X)
++AxisMappings=(AxisName="MovementAxisRight_X",Scale=1.000000,Key=MixedReality_Right_Thumbstick_X)
++AxisMappings=(AxisName="MovementAxisRight_X",Scale=1.000000,Key=Vive_Right_Trackpad_X)
++AxisMappings=(AxisName="MovementAxisRight_Y",Scale=1.000000,Key=OculusTouch_Right_Thumbstick_Y)
++AxisMappings=(AxisName="MovementAxisRight_Y",Scale=1.000000,Key=ValveIndex_Right_Thumbstick_Y)
++AxisMappings=(AxisName="MovementAxisRight_Y",Scale=1.000000,Key=MixedReality_Right_Thumbstick_Y)
++AxisMappings=(AxisName="MovementAxisRight_Y",Scale=1.000000,Key=Vive_Right_Trackpad_Y)
++AxisMappings=(AxisName="MovementAxisRight_Y",Scale=1.000000,Key=OculusTouch_Right_A_Click)
++AxisMappings=(AxisName="SpectatorYaw",Scale=1.000000,Key=Gamepad_RightX)
++AxisMappings=(AxisName="SpectatorPitch",Scale=-1.000000,Key=Gamepad_RightY)
++AxisMappings=(AxisName="SpectatorMoveForward",Scale=1.000000,Key=Gamepad_LeftY)
++AxisMappings=(AxisName="SpectatorMoveRight",Scale=1.000000,Key=Gamepad_LeftX)
++AxisMappings=(AxisName="SpectatorMoveUp",Scale=1.000000,Key=Gamepad_RightShoulder)
++AxisMappings=(AxisName="SpectatorMoveUp",Scale=-1.000000,Key=Gamepad_LeftShoulder)
++AxisMappings=(AxisName="SpectatorFOV",Scale=1.000000,Key=Gamepad_RightTriggerAxis)
++AxisMappings=(AxisName="SpectatorFOV",Scale=-1.000000,Key=Gamepad_LeftTriggerAxis)
 
 [Redirects]
 GameName=TP_VirtualRealityBP
diff --git a/Templates/TP_VirtualRealityBP/manifest.json b/Templates/TP_VirtualRealityBP/manifest.json
--- a/Templates/TP_VirtualRealityBP/manifest.json
+++ b/Templates/TP_VirtualRealityBP/manifest.json
@@ -68,8 +68,13 @@
 		},
 		{
 			"MountName": "FPWeapon",
-			"DetailLevels": 
-			[
+			"DetailLevels": [
+				"Standard"
+			]
+		},
+		{
+			"MountName": "VRSpectator",
+			"DetailLevels": [
 				"Standard"
 			]
 		}
