XR | VR Template now has the Oculus VR Plugin enabled by default and the XR API is set to Native OpenXR

diff --git a/Templates/TP_VirtualRealityBP/Config/DefaultEngine.ini b/Templates/TP_VirtualRealityBP/Config/DefaultEngine.ini
--- a/Templates/TP_VirtualRealityBP/Config/DefaultEngine.ini
+++ b/Templates/TP_VirtualRealityBP/Config/DefaultEngine.ini
@@ -285,6 +285,7 @@ ReverbPlugin=Resonance Audio
 [/Script/OculusHMD.OculusHMDRuntimeSettings]
 HandTrackingSupport=ControllersAndHands
 FFRLevel=FFR_Medium
+XrApi=NativeOpenXR
 
 [/Script/Engine.CollisionProfile]
 -Profiles=(Name="NoCollision",CollisionEnabled=NoCollision,ObjectTypeName="WorldStatic",CustomResponses=((Channel="Visibility",Response=ECR_Ignore),(Channel="Camera",Response=ECR_Ignore)),HelpMessage="No collision",bCanModify=False)
diff --git a/Templates/TP_VirtualRealityBP/TP_VirtualRealityBP.uproject b/Templates/TP_VirtualRealityBP/TP_VirtualRealityBP.uproject
--- a/Templates/TP_VirtualRealityBP/TP_VirtualRealityBP.uproject
+++ b/Templates/TP_VirtualRealityBP/TP_VirtualRealityBP.uproject
@@ -12,15 +12,6 @@
 			"Name": "ResonanceAudio",
 			"Enabled": true
 		},
-		{
-			"Name": "OculusVR",
-			"Enabled": false,
-			"SupportedTargetPlatforms": [
-				"Win32",
-				"Win64",
-				"Android"
-			]
-		},
 		{
 			"Name": "OpenXR",
 			"Enabled": true,
