XR | VR Template | Disable Android OpenGLES3.1 support

diff --git a/Templates/TP_VirtualRealityBP/Config/DefaultEngine.ini b/Templates/TP_VirtualRealityBP/Config/DefaultEngine.ini
--- a/Templates/TP_VirtualRealityBP/Config/DefaultEngine.ini
+++ b/Templates/TP_VirtualRealityBP/Config/DefaultEngine.ini
@@ -215,7 +215,7 @@ KeyPassword=
 bBuildForArmV7=True
 bBuildForArm64=False
 bBuildForX8664=False
-bBuildForES31=True
+bBuildForES31=False
 bSupportsVulkan=True
 bSupportsVulkanSM5=False
 ClangSanitizer=None
