XR | VR Template | Level refresh with new level prototyping content

diff --git a/Templates/TP_VirtualRealityBP/Config/TemplateDefs.ini b/Templates/TP_VirtualRealityBP/Config/TemplateDefs.ini
--- a/Templates/TP_VirtualRealityBP/Config/TemplateDefs.ini
+++ b/Templates/TP_VirtualRealityBP/Config/TemplateDefs.ini
@@ -41,9 +41,9 @@ ReplacementsInFiles=(Extensions=("cpp","h","ini","cs"),From="%TEMPLATENAME_UPPER
 ReplacementsInFiles=(Extensions=("cpp","h","ini","cs"),From="%TEMPLATENAME_LOWERCASE%",To="%PROJECTNAME_LOWERCASE%",bCaseSensitive=true)
 ReplacementsInFiles=(Extensions=("cpp","h","ini","cs"),From="%TEMPLATENAME%",To="%PROJECTNAME%",bCaseSensitive=false)
 SortKey=VirtualRealityBP
-SharedContentPacks=(MountName="Geometry",DetailLevels=("Standard","High"))
 SharedContentPacks=(MountName="FPWeapon",DetailLevels=("Standard"))
 SharedContentPacks=(MountName="VRSpectator",DetailLevels=("Standard"))
+SharedContentPacks=(MountName="LevelPrototyping",DetailLevels=("High"))
 EditDetailLevelPreference="High"
 FocusAsset=/Game/VRTemplate/Blueprints
 
diff --git a/Templates/TP_VirtualRealityBP/manifest.json b/Templates/TP_VirtualRealityBP/manifest.json
--- a/Templates/TP_VirtualRealityBP/manifest.json
+++ b/Templates/TP_VirtualRealityBP/manifest.json
@@ -59,10 +59,9 @@
 	"AdditionalFeaturePacks": 
 	[
 		{
-			"MountName": "Geometry",
+			"MountName": "LevelPrototyping",
 			"DetailLevels": 
 			[
-				"Standard",
 				"High"
 			]
 		},
