Updated VirtualRealityBP TemplateDefs.ini to update the Chinese description for this template

diff --git a/Templates/TP_VirtualRealityBP/Config/TemplateDefs.ini b/Templates/TP_VirtualRealityBP/Config/TemplateDefs.ini
--- a/Templates/TP_VirtualRealityBP/Config/TemplateDefs.ini
+++ b/Templates/TP_VirtualRealityBP/Config/TemplateDefs.ini
@@ -12,7 +12,7 @@ LocalizedDisplayNames=(Language="ja",Text="バーチャルリアリティ")
 LocalizedDescriptions=(Language="ja",Text="デスクトップ VR と Playstation VR 向け、ブループリントのバーチャルリアリティテンプレートです。テンプレートには二種類のテレポートロコモーションモードがあります。それぞれに ポーンとマップが用意されています。MotionControllerMap ではモーションコントローラーを使ったテレポートが可能です。HMDLocomotionMap ではゲームパッドを使ったテレポート、及び、ヘッドマウントディスプレイを使ったポイントアンドテレポートが可能です。")
 
 LocalizedDisplayNames=(Language="zh-Hans",Text="虚拟现实应用")
-LocalizedDescriptions=(Language="zh-Hans",Text="蓝图虚拟现实应用模板面向台式PC以及Playstation VR设备。模板包含2中传送移动模式，分别拥有独立的Pawn和地图。MotionControllerMap可以让玩家使用体感手柄四处传送，而HMDLocomotionMap可以让玩家使用游戏手柄或头显设备进行“指向-传送”式的传送。")
+LocalizedDescriptions=(Language="zh-Hans",Text="蓝图虚拟现实模板基于OpenXR实现，用于为面向台式电脑、主机端以及移动端的VR设备开发项目。该模板默认实现了玩家传送、玩家旋转、物体拾取、物体交互、Google Resonance空间化音频，以及VR观察视角等功能。")
 
 ClassTypes=Pawn, GameMode, Blueprint Interface, Actor
 AssetTypes=Sound Cue, Haptic Feedback Effect Curve, Niagara System.
