Updated VR Template thumbnail and added line for full size icon in TemplateDefs.ini

diff --git a/Templates/TP_VirtualRealityBP/Config/TemplateDefs.ini b/Templates/TP_VirtualRealityBP/Config/TemplateDefs.ini
--- a/Templates/TP_VirtualRealityBP/Config/TemplateDefs.ini
+++ b/Templates/TP_VirtualRealityBP/Config/TemplateDefs.ini
@@ -47,6 +47,9 @@ SharedContentPacks=(MountName="LevelPrototyping",DetailLevels=("High"))
 EditDetailLevelPreference="High"
 FocusAsset=/Game/VRTemplate/Blueprints
 
+; This project uses a thumbnail as the project icon
+bThumbnailAsIcon=true
+
 HiddenSettings=RayTracing
 HiddenSettings=Languages
 HiddenSettings=HardwareTarget
