Removed Oculus Touch (R) A Press binding for MovementAxisRight_Y as it's causing issues if another Action Mapping is using the button

diff --git a/Templates/TP_VirtualRealityBP/Config/DefaultInput.ini b/Templates/TP_VirtualRealityBP/Config/DefaultInput.ini
--- a/Templates/TP_VirtualRealityBP/Config/DefaultInput.ini
+++ b/Templates/TP_VirtualRealityBP/Config/DefaultInput.ini
@@ -81,6 +81,8 @@
 +AxisConfig=(AxisKeyName="OculusHand_Right_MiddlePinchStrength",AxisProperties=(DeadZone=0.000000,Sensitivity=1.000000,Exponent=1.000000,bInvert=False))
 +AxisConfig=(AxisKeyName="OculusHand_Right_RingPinchStrength",AxisProperties=(DeadZone=0.000000,Sensitivity=1.000000,Exponent=1.000000,bInvert=False))
 +AxisConfig=(AxisKeyName="OculusHand_Right_PinkPinchStrength",AxisProperties=(DeadZone=0.000000,Sensitivity=1.000000,Exponent=1.000000,bInvert=False))
++AxisConfig=(AxisKeyName="OculusTouch_Left_ThumbRest",AxisProperties=(DeadZone=0.000000,Sensitivity=1.000000,Exponent=1.000000,bInvert=False))
++AxisConfig=(AxisKeyName="OculusTouch_Right_ThumbRest",AxisProperties=(DeadZone=0.000000,Sensitivity=1.000000,Exponent=1.000000,bInvert=False))
 bAltEnterTogglesFullscreen=True
 bF11TogglesFullscreen=True
 bUseMouseForTouch=False
@@ -169,7 +171,6 @@ DoubleClickTime=0.200000
 +AxisMappings=(AxisName="MovementAxisRight_Y",Scale=1.000000,Key=ValveIndex_Right_Thumbstick_Y)
 +AxisMappings=(AxisName="MovementAxisRight_Y",Scale=1.000000,Key=MixedReality_Right_Thumbstick_Y)
 +AxisMappings=(AxisName="MovementAxisRight_Y",Scale=1.000000,Key=Vive_Right_Trackpad_Y)
-+AxisMappings=(AxisName="MovementAxisRight_Y",Scale=1.000000,Key=OculusTouch_Right_A_Click)
 +AxisMappings=(AxisName="SpectatorYaw",Scale=1.000000,Key=Gamepad_RightX)
 +AxisMappings=(AxisName="SpectatorPitch",Scale=-1.000000,Key=Gamepad_RightY)
 +AxisMappings=(AxisName="SpectatorMoveForward",Scale=1.000000,Key=Gamepad_LeftY)
