bBuildForArm64 no longer defaults to false when creating a new VR Template. It was set in TP_Virtualreality project settings but not in the .ini!?

diff --git a/Templates/TP_VirtualRealityBP/Config/DefaultEngine.ini b/Templates/TP_VirtualRealityBP/Config/DefaultEngine.ini
--- a/Templates/TP_VirtualRealityBP/Config/DefaultEngine.ini
+++ b/Templates/TP_VirtualRealityBP/Config/DefaultEngine.ini
@@ -212,8 +212,7 @@ KeyStore=UE.keystore
 KeyAlias=DistroKey
 KeyStorePassword=codesigner
 KeyPassword=
-bBuildForArmV7=True
-bBuildForArm64=False
+bBuildForArm64=True
 bBuildForX8664=False
 bBuildForES31=False
 bSupportsVulkan=True
