VR Template | Allow client side navigation is now set to true. This will prevent confusion when users are trying to convert the template to work in multiplayer.

diff --git a/Templates/TP_VirtualRealityBP/Config/DefaultEngine.ini b/Templates/TP_VirtualRealityBP/Config/DefaultEngine.ini
--- a/Templates/TP_VirtualRealityBP/Config/DefaultEngine.ini
+++ b/Templates/TP_VirtualRealityBP/Config/DefaultEngine.ini
@@ -343,3 +343,6 @@ XrApi=NativeOpenXR
 +CollisionChannelRedirects=(OldName="VehicleMovement",NewName="Vehicle")
 +CollisionChannelRedirects=(OldName="PawnMovement",NewName="Pawn")
 
+[/Script/NavigationSystem.NavigationSystemV1]
+bAllowClientSideNavigation=True
+
