XR | VR Template| Disabling Mobile Mesh-Auto Instancing as it's causing performance issues on Quest 1

diff --git a/Templates/TP_VirtualRealityBP/Config/DefaultEngine.ini b/Templates/TP_VirtualRealityBP/Config/DefaultEngine.ini
--- a/Templates/TP_VirtualRealityBP/Config/DefaultEngine.ini
+++ b/Templates/TP_VirtualRealityBP/Config/DefaultEngine.ini
@@ -112,7 +112,6 @@ bDiscardSkeletalMeshOptionalLODs=(Default=False,PerPlatform=())
 VisualizeCalibrationColorMaterialPath=None
 VisualizeCalibrationCustomMaterialPath=None
 VisualizeCalibrationGrayscaleMaterialPath=None
-r.Mobile.SupportGPUScene=1
 r.Mobile.AntiAliasing=3
 
 [/Script/HardwareTargeting.HardwareTargetingSettings]
@@ -344,5 +343,4 @@ XrApi=NativeOpenXR
 +CollisionChannelRedirects=(OldName="PawnMovement",NewName="Pawn")
 
 [/Script/NavigationSystem.NavigationSystemV1]
-bAllowClientSideNavigation=True
-
+bAllowClientSideNavigation=True
\ No newline at end of file
