Fix deprecated CVars in VR Template.

r.BasePassOutputsVelocity=0 -> r.VelocityOutputPass=0
r.VertexDeformationOutputsVelocity=0 -> r.Velocity.EnableVertexDeformation=0

diff --git a/Templates/TP_VirtualRealityBP/Config/DefaultEngine.ini b/Templates/TP_VirtualRealityBP/Config/DefaultEngine.ini
--- a/Templates/TP_VirtualRealityBP/Config/DefaultEngine.ini
+++ b/Templates/TP_VirtualRealityBP/Config/DefaultEngine.ini
@@ -59,8 +59,8 @@ r.EarlyZPass=3
 r.EarlyZPassOnlyMaterialMasking=False
 r.DBuffer=True
 r.ClearSceneMethod=1
-r.BasePassOutputsVelocity=False
-r.VertexDeformationOutputsVelocity=False
+r.VelocityOutputPass=0
+r.Velocity.EnableVertexDeformation=0
 r.SelectiveBasePassOutputs=False
 bDefaultParticleCutouts=False
 fx.GPUSimulationTextureSizeX=1024
