XR | VR Template | Now using Interaural Time Delay for Audio Spatialization and Reverb instead of Resonance Audio. Removed config line related to deprecated platform "Lumin".

diff --git a/Templates/TP_VirtualRealityBP/Config/DefaultEngine.ini b/Templates/TP_VirtualRealityBP/Config/DefaultEngine.ini
--- a/Templates/TP_VirtualRealityBP/Config/DefaultEngine.ini
+++ b/Templates/TP_VirtualRealityBP/Config/DefaultEngine.ini
@@ -150,8 +150,8 @@ AudioCallbackBufferFrameSize=1024
 AudioNumBuffersToEnqueue=1
 AudioMaxChannels=0
 AudioNumSourceWorkers=4
-SpatializationPlugin=Resonance Audio
-ReverbPlugin=Resonance Audio
+SpatializationPlugin=Built-in Spatialization
+ReverbPlugin=Built-in Reverb
 OcclusionPlugin=
 CompressionOverrides=(bOverrideCompressionTimes=False,DurationThreshold=5.000000,MaxNumRandomBranches=0,SoundCueQualityIndex=0)
 bUseAudioStreamCaching=False
@@ -244,8 +244,8 @@ AudioCallbackBufferFrameSize=1024
 AudioNumBuffersToEnqueue=4
 AudioMaxChannels=0
 AudioNumSourceWorkers=0
-SpatializationPlugin=Resonance Audio
-ReverbPlugin=Resonance Audio
+SpatializationPlugin=Built-in Spatialization
+ReverbPlugin=Built-in Reverb
 OcclusionPlugin=
 CompressionOverrides=(bOverrideCompressionTimes=False,DurationThreshold=5.000000,MaxNumRandomBranches=0,SoundCueQualityIndex=0)
 bUseAudioStreamCaching=False
@@ -273,21 +273,6 @@ SDKAPILevelOverride=
 NDKAPILevelOverride=
 bStreamLandscapeMeshLODs=False
 
-[/Script/IOSRuntimeSettings.IOSRuntimeSettings]
-SpatializationPlugin=Resonance Audio
-
-[/Script/LuminRuntimeSettings.LuminRuntimeSettings]
-SpatializationPlugin=Built-in Spatialization
-
-[/Script/MacTargetPlatform.MacTargetSettings]
-SpatializationPlugin=Resonance Audio
-ReverbPlugin=Resonance Audio
-
-[/Script/OculusHMD.OculusHMDRuntimeSettings]
-HandTrackingSupport=ControllersAndHands
-FFRLevel=FFR_Medium
-XrApi=NativeOpenXR
-
 [/Script/Engine.CollisionProfile]
 -Profiles=(Name="NoCollision",CollisionEnabled=NoCollision,ObjectTypeName="WorldStatic",CustomResponses=((Channel="Visibility",Response=ECR_Ignore),(Channel="Camera",Response=ECR_Ignore)),HelpMessage="No collision",bCanModify=False)
 -Profiles=(Name="BlockAll",CollisionEnabled=QueryAndPhysics,ObjectTypeName="WorldStatic",CustomResponses=,HelpMessage="WorldStatic object that blocks all actors by default. All new custom channels will use its own default response. ",bCanModify=False)
@@ -346,4 +331,19 @@ XrApi=NativeOpenXR
 +CollisionChannelRedirects=(OldName="PawnMovement",NewName="Pawn")
 
 [/Script/NavigationSystem.NavigationSystemV1]
-bAllowClientSideNavigation=True
\ No newline at end of file
+bAllowClientSideNavigation=True
+
+[/Script/AndroidFileServerEditor.AndroidFileServerRuntimeSettings]
+bEnablePlugin=True
+bAllowNetworkConnection=True
+SecurityToken=9AE2A8ED4820B9C3D9A97F96ADA01BBC
+bIncludeInShipping=False
+bAllowExternalStartInShipping=False
+bCompileAFSProject=False
+bUseCompression=False
+bLogFiles=False
+bReportStats=False
+ConnectionType=USBOnly
+bUseManualIPAddress=False
+ManualIPAddress=
+
diff --git a/Templates/TP_VirtualRealityBP/TP_VirtualRealityBP.uproject b/Templates/TP_VirtualRealityBP/TP_VirtualRealityBP.uproject
--- a/Templates/TP_VirtualRealityBP/TP_VirtualRealityBP.uproject
+++ b/Templates/TP_VirtualRealityBP/TP_VirtualRealityBP.uproject
@@ -4,10 +4,6 @@
 	"Category": "",
 	"Description": "",
 	"Plugins": [
-		{
-			"Name": "ResonanceAudio",
-			"Enabled": true
-		},
 		{
 			"Name": "OpenXR",
 			"Enabled": true,
