Disabled DXC on VR template due to bug with MatrixStride in SPIRV (fixed in 5.2)

diff --git a/Templates/TP_VirtualRealityBP/Config/DefaultEngine.ini b/Templates/TP_VirtualRealityBP/Config/DefaultEngine.ini
--- a/Templates/TP_VirtualRealityBP/Config/DefaultEngine.ini
+++ b/Templates/TP_VirtualRealityBP/Config/DefaultEngine.ini
@@ -111,6 +111,8 @@ VisualizeCalibrationColorMaterialPath=None
 VisualizeCalibrationCustomMaterialPath=None
 VisualizeCalibrationGrayscaleMaterialPath=None
 r.Mobile.AntiAliasing=3
+r.Mobile.FloatPrecisionMode=2
+r.OpenGL.ForceDXC=0
 
 [/Script/HardwareTargeting.HardwareTargetingSettings]
 TargetedHardwareClass=Mobile
