Removing unique Android File Server settings

diff --git a/Templates/TP_VirtualRealityBP/Config/DefaultEngine.ini b/Templates/TP_VirtualRealityBP/Config/DefaultEngine.ini
--- a/Templates/TP_VirtualRealityBP/Config/DefaultEngine.ini
+++ b/Templates/TP_VirtualRealityBP/Config/DefaultEngine.ini
@@ -333,19 +333,4 @@ bStreamLandscapeMeshLODs=False
 +CollisionChannelRedirects=(OldName="PawnMovement",NewName="Pawn")
 
 [/Script/NavigationSystem.NavigationSystemV1]
-bAllowClientSideNavigation=True
-
-[/Script/AndroidFileServerEditor.AndroidFileServerRuntimeSettings]
-bEnablePlugin=True
-bAllowNetworkConnection=True
-SecurityToken=9AE2A8ED4820B9C3D9A97F96ADA01BBC
-bIncludeInShipping=False
-bAllowExternalStartInShipping=False
-bCompileAFSProject=False
-bUseCompression=False
-bLogFiles=False
-bReportStats=False
-ConnectionType=USBOnly
-bUseManualIPAddress=False
-ManualIPAddress=
-
+bAllowClientSideNavigation=True
\ No newline at end of file
