Replace list of quest devices in packaging settings with checkbox to package for all quest devices.

Also updating VRTemplate to create projects with the new settings and removing daydream deploy logic.

diff --git a/Templates/TP_VirtualRealityBP/Config/DefaultEngine.ini b/Templates/TP_VirtualRealityBP/Config/DefaultEngine.ini
--- a/Templates/TP_VirtualRealityBP/Config/DefaultEngine.ini
+++ b/Templates/TP_VirtualRealityBP/Config/DefaultEngine.ini
@@ -207,8 +207,7 @@ bBundleDensitySplit=True
 ExtraApplicationSettings=
 ExtraActivitySettings=
 bAndroidVoiceEnabled=False
-+PackageForOculusMobile=Quest2
-+PackageForOculusMobile=Quest
+bPackageForMetaQuest=True
 bRemoveOSIG=False
 +GoogleVRCaps=Daydream33
 bGoogleVRSustainedPerformance=False
