VRTemplate: disable ISR on Mac

#rnx

diff --git a/Templates/TP_VirtualRealityBP/Config/Mac/MacEngine.ini b/Templates/TP_VirtualRealityBP/Config/Mac/MacEngine.ini
new file mode 100644
--- /dev/null
+++ b/Templates/TP_VirtualRealityBP/Config/Mac/MacEngine.ini
@@ -0,0 +1,4 @@
+
+[/Script/Engine.RendererSettings]
+vr.InstancedStereo=False
+vr.MobileMultiView=False
