Disable local exposure in templates that already had auto exposure disabled. (5.3 now defaults it on in new projects.)

diff --git a/Templates/TP_AEC_HandheldARBP/Config/DefaultEngine.ini b/Templates/TP_AEC_HandheldARBP/Config/DefaultEngine.ini
--- a/Templates/TP_AEC_HandheldARBP/Config/DefaultEngine.ini
+++ b/Templates/TP_AEC_HandheldARBP/Config/DefaultEngine.ini
@@ -36,6 +36,8 @@ r.DefaultFeature.AmbientOcclusion=False
 r.DefaultFeature.AmbientOcclusionStaticFraction=True
 r.DefaultFeature.AutoExposure=False
 r.DefaultFeature.AutoExposure.Method=0
+r.DefaultFeature.LocalExposure.HighlightContrastScale=1.0
+r.DefaultFeature.LocalExposure.ShadowContrastScale=1.0
 r.DefaultFeature.MotionBlur=False
 r.DefaultFeature.LensFlare=False
 r.TemporalAA.Upsampling=False
diff --git a/Templates/TP_DMXBP/Config/DefaultEngine.ini b/Templates/TP_DMXBP/Config/DefaultEngine.ini
--- a/Templates/TP_DMXBP/Config/DefaultEngine.ini
+++ b/Templates/TP_DMXBP/Config/DefaultEngine.ini
@@ -8,6 +8,8 @@ r.GenerateMeshDistanceFields=False
 r.AllowVolumeTextureAssetCreation=1
 r.DefaultFeature.AutoExposure=True
 r.DefaultFeature.Bloom=True
+r.DefaultFeature.LocalExposure.HighlightContrastScale=1.0
+r.DefaultFeature.LocalExposure.ShadowContrastScale=1.0
 r.SSGI.Enable=True
 r.DefaultFeature.AutoExposure.Method=0
 r.DefaultFeature.LensFlare=True
diff --git a/Templates/TP_HandheldARBP/Config/DefaultEngine.ini b/Templates/TP_HandheldARBP/Config/DefaultEngine.ini
--- a/Templates/TP_HandheldARBP/Config/DefaultEngine.ini
+++ b/Templates/TP_HandheldARBP/Config/DefaultEngine.ini
@@ -33,6 +33,8 @@ r.DefaultFeature.AmbientOcclusion=False
 r.DefaultFeature.AmbientOcclusionStaticFraction=True
 r.DefaultFeature.AutoExposure=False
 r.DefaultFeature.AutoExposure.Method=0
+r.DefaultFeature.LocalExposure.HighlightContrastScale=1.0
+r.DefaultFeature.LocalExposure.ShadowContrastScale=1.0
 r.DefaultFeature.MotionBlur=False
 r.DefaultFeature.LensFlare=False
 r.TemporalAA.Upsampling=False
diff --git a/Templates/TP_InCamVFXBP/Config/DefaultEngine.ini b/Templates/TP_InCamVFXBP/Config/DefaultEngine.ini
--- a/Templates/TP_InCamVFXBP/Config/DefaultEngine.ini
+++ b/Templates/TP_InCamVFXBP/Config/DefaultEngine.ini
@@ -14,6 +14,8 @@ bEnabled=True
 [/Script/Engine.RendererSettings]
 r.DefaultFeature.Bloom=False
 r.DefaultFeature.AutoExposure=False
+r.DefaultFeature.LocalExposure.HighlightContrastScale=1.0
+r.DefaultFeature.LocalExposure.ShadowContrastScale=1.0
 r.DefaultFeature.MotionBlur=False
 r.SceneRenderTargetResizeMethod=2
 r.DefaultBackBufferPixelFormat=4
diff --git a/Templates/TP_ME_BlankBP/Config/DefaultEngine.ini b/Templates/TP_ME_BlankBP/Config/DefaultEngine.ini
--- a/Templates/TP_ME_BlankBP/Config/DefaultEngine.ini
+++ b/Templates/TP_ME_BlankBP/Config/DefaultEngine.ini
@@ -12,6 +12,8 @@ r.Tonemapper.Sharpen=2
 r.DefaultFeature.AutoExposure.ExtendDefaultLuminanceRange=True
 r.DefaultFeature.Bloom=False
 r.DefaultFeature.AutoExposure=False
+r.DefaultFeature.LocalExposure.HighlightContrastScale=1.0
+r.DefaultFeature.LocalExposure.ShadowContrastScale=1.0
 r.DefaultFeature.MotionBlur=False
 r.SceneRenderTargetResizeMethod=2
 
diff --git a/Templates/TP_ME_VProdBP/Config/DefaultEngine.ini b/Templates/TP_ME_VProdBP/Config/DefaultEngine.ini
--- a/Templates/TP_ME_VProdBP/Config/DefaultEngine.ini
+++ b/Templates/TP_ME_VProdBP/Config/DefaultEngine.ini
@@ -11,6 +11,8 @@ r.Tonemapper.Sharpen=2
 r.DefaultFeature.AutoExposure.ExtendDefaultLuminanceRange=True
 r.DefaultFeature.Bloom=False
 r.DefaultFeature.AutoExposure=False
+r.DefaultFeature.LocalExposure.HighlightContrastScale=1.0
+r.DefaultFeature.LocalExposure.ShadowContrastScale=1.0
 r.DefaultFeature.MotionBlur=False
 r.SceneRenderTargetResizeMethod=2
 
diff --git a/Templates/TP_MFG_HandheldARBP/Config/DefaultEngine.ini b/Templates/TP_MFG_HandheldARBP/Config/DefaultEngine.ini
--- a/Templates/TP_MFG_HandheldARBP/Config/DefaultEngine.ini
+++ b/Templates/TP_MFG_HandheldARBP/Config/DefaultEngine.ini
@@ -35,6 +35,8 @@ r.DefaultFeature.AmbientOcclusion=False
 r.DefaultFeature.AmbientOcclusionStaticFraction=True
 r.DefaultFeature.AutoExposure=False
 r.DefaultFeature.AutoExposure.Method=0
+r.DefaultFeature.LocalExposure.HighlightContrastScale=1.0
+r.DefaultFeature.LocalExposure.ShadowContrastScale=1.0
 r.DefaultFeature.MotionBlur=False
 r.DefaultFeature.LensFlare=False
 r.TemporalAA.Upsampling=False
diff --git a/Templates/TP_MFG_ProdConfigBP/Config/DefaultEngine.ini b/Templates/TP_MFG_ProdConfigBP/Config/DefaultEngine.ini
--- a/Templates/TP_MFG_ProdConfigBP/Config/DefaultEngine.ini
+++ b/Templates/TP_MFG_ProdConfigBP/Config/DefaultEngine.ini
@@ -13,6 +13,8 @@ AppliedDefaultGraphicsPerformance=Maximum
 [/Script/Engine.RendererSettings]
 r.DefaultFeature.AutoExposure.ExtendDefaultLuminanceRange=True
 r.DefaultFeature.AutoExposure=False
+r.DefaultFeature.LocalExposure.HighlightContrastScale=1.0
+r.DefaultFeature.LocalExposure.ShadowContrastScale=1.0
 r.GenerateMeshDistanceFields=True
 
 [/Script/EngineSettings.GameMapsSettings]
diff --git a/Templates/TP_VirtualRealityBP/Config/DefaultEngine.ini b/Templates/TP_VirtualRealityBP/Config/DefaultEngine.ini
--- a/Templates/TP_VirtualRealityBP/Config/DefaultEngine.ini
+++ b/Templates/TP_VirtualRealityBP/Config/DefaultEngine.ini
@@ -46,6 +46,8 @@ r.DefaultFeature.AutoExposure.Method=0
 r.DefaultFeature.AutoExposure.Bias=1.000000
 r.DefaultFeature.AutoExposure.ExtendDefaultLuminanceRange=True
 r.EyeAdaptation.EditorOnly=False
+r.DefaultFeature.LocalExposure.HighlightContrastScale=1.0
+r.DefaultFeature.LocalExposure.ShadowContrastScale=1.0
 r.DefaultFeature.MotionBlur=False
 r.DefaultFeature.LensFlare=False
 r.TemporalAA.Upsampling=False
diff --git a/Templates/TP_nDisplayBP/Config/DefaultEngine.ini b/Templates/TP_nDisplayBP/Config/DefaultEngine.ini
--- a/Templates/TP_nDisplayBP/Config/DefaultEngine.ini
+++ b/Templates/TP_nDisplayBP/Config/DefaultEngine.ini
@@ -14,6 +14,8 @@ bEnabled=True
 [/Script/Engine.RendererSettings]
 r.DefaultFeature.Bloom=False
 r.DefaultFeature.AutoExposure=False
+r.DefaultFeature.LocalExposure.HighlightContrastScale=1.0
+r.DefaultFeature.LocalExposure.ShadowContrastScale=1.0
 r.DefaultFeature.MotionBlur=False
 r.SceneRenderTargetResizeMethod=2
 vr.AllowMotionBlurInVR=True
