XR | VR Template | Updated DefaultEngine.ini to better match BaseEngine.ini. Many of the previous cvars were left overs from UE4.

diff --git a/Templates/TP_VirtualRealityBP/Config/DefaultEngine.ini b/Templates/TP_VirtualRealityBP/Config/DefaultEngine.ini
--- a/Templates/TP_VirtualRealityBP/Config/DefaultEngine.ini
+++ b/Templates/TP_VirtualRealityBP/Config/DefaultEngine.ini
@@ -1,337 +1,71 @@
-
+[/Script/EngineSettings.GameMapsSettings]
+GlobalDefaultGameMode=/Game/VRTemplate/Blueprints/VRGameMode.VRGameMode_C
+EditorStartupMap=/Game/VRTemplate/Maps/VRTemplateMap.VRTemplateMap
+GameDefaultMap=/Game/VRTemplate/Maps/VRTemplateMap.VRTemplateMap
 
 [/Script/Engine.RendererSettings]
-r.Mobile.DisableVertexFog=True
-r.Mobile.AllowDitheredLODTransition=False
-r.Mobile.AllowSoftwareOcclusion=False
-r.Mobile.VirtualTextures=False
-r.DiscardUnusedQuality=False
-r.AllowOcclusionQueries=True
-r.MinScreenRadiusForLights=0.030000
-r.MinScreenRadiusForDepthPrepass=0.030000
-r.MinScreenRadiusForCSMDepth=0.010000
-r.PrecomputedVisibilityWarning=False
-r.TextureStreaming=True
-Compat.UseDXT5NormalMaps=False
-r.VirtualTextures=False
-r.VirtualTexturedLightmaps=False
-r.VT.TileSize=128
-r.VT.TileBorderSize=4
-r.vt.FeedbackFactor=16
-r.VT.EnableCompressZlib=True
-r.VT.EnableCompressCrunch=False
-r.ClearCoatNormal=False
-r.ReflectionCaptureResolution=128
-r.ReflectionEnvironmentLightmapMixBasedOnRoughness=True
 r.ForwardShading=True
-r.VertexFoggingForOpaque=True
-r.AllowStaticLighting=True
-r.NormalMapsForStaticLighting=False
-r.GenerateMeshDistanceFields=False
-r.DistanceFieldBuild.EightBit=False
-r.GenerateLandscapeGIData=False
-r.DistanceFieldBuild.Compress=False
-r.TessellationAdaptivePixelsPerTriangle=48.000000
-r.SeparateTranslucency=False
-r.TranslucentSortPolicy=0
-TranslucentSortAxis=(X=0.000000,Y=-1.000000,Z=0.000000)
-r.CustomDepth=1
-r.CustomDepthTemporalAAJitter=True
-r.PostProcessing.PropagateAlpha=0
-r.DefaultFeature.Bloom=False
-r.DefaultFeature.AmbientOcclusion=False
-r.DefaultFeature.AmbientOcclusionStaticFraction=True
-r.DefaultFeature.AutoExposure=False
-r.DefaultFeature.AutoExposure.Method=0
-r.DefaultFeature.AutoExposure.Bias=1.000000
-r.DefaultFeature.AutoExposure.ExtendDefaultLuminanceRange=True
-r.EyeAdaptation.EditorOnly=False
-r.DefaultFeature.LocalExposure.HighlightContrastScale=1.0
-r.DefaultFeature.LocalExposure.ShadowContrastScale=1.0
-r.DefaultFeature.MotionBlur=False
-r.DefaultFeature.LensFlare=False
-r.TemporalAA.Upsampling=False
-r.SSGI.Enable=False
-r.AntiAliasingMethod=3
-r.DefaultFeature.LightUnits=1
-r.DefaultBackBufferPixelFormat=4
-r.Shadow.UnbuiltPreviewInGame=True
-r.StencilForLODDither=False
-r.EarlyZPass=3
-r.EarlyZPassOnlyMaterialMasking=False
-r.DBuffer=True
-r.ClearSceneMethod=1
-r.VelocityOutputPass=0
-r.Velocity.EnableVertexDeformation=0
-r.SelectiveBasePassOutputs=False
-bDefaultParticleCutouts=False
-fx.GPUSimulationTextureSizeX=1024
-fx.GPUSimulationTextureSizeY=1024
-r.AllowGlobalClipPlane=False
-r.GBufferFormat=1
-r.MorphTarget.Mode=True
-r.GPUCrashDebugging=False
-vr.InstancedStereo=True
 r.MobileHDR=False
 vr.MobileMultiView=True
+vr.InstancedStereo=True
+r.DefaultFeature.AutoExposure=False
+r.DefaultFeature.AmbientOcclusion=False
+r.DefaultFeature.AmbientOcclusionStaticFraction=False
+r.DefaultFeature.MotionBlur=False
 r.Mobile.UseHWsRGBEncoding=True
-vr.RoundRobinOcclusion=False
-vr.ODSCapture=False
-r.MeshStreaming=False
-r.WireframeCullThreshold=5.000000
-r.RayTracing=False
-r.RayTracing.UseTextureLod=False
-r.SupportStationarySkylight=True
-r.SupportLowQualityLightmaps=True
-r.SupportPointLightWholeSceneShadows=True
-r.SupportAtmosphericFog=True
-r.SupportSkyAtmosphere=True
-r.SupportSkyAtmosphereAffectsHeightFog=False
-r.SkinCache.CompileShaders=False
-r.SkinCache.DefaultBehavior=1
-r.SkinCache.SceneMemoryLimitInMB=128.000000
-r.Mobile.EnableStaticAndCSMShadowReceivers=True
-r.Mobile.EnableMovableLightCSMShaderCulling=True
-r.Mobile.AllowDistanceFieldShadows=True
-r.Mobile.AllowMovableDirectionalLights=True
-r.MobileNumDynamicPointLights=4
-r.MobileDynamicPointLightsUseStaticBranch=True
-r.Mobile.EnableMovableSpotlights=False
-r.Mobile.EnableMovableSpotlightsShadow=False
-r.GPUSkin.Support16BitBoneIndex=False
-r.GPUSkin.Limit2BoneInfluences=False
-r.SupportDepthOnlyIndexBuffers=True
-r.SupportReversedIndexBuffers=True
-r.LightPropagationVolume=False
-r.Mobile.AmbientOcclusion=False
-r.GPUSkin.UnlimitedBoneInfluences=False
-r.GPUSkin.UnlimitedBoneInfluencesThreshold=8
-r.Mobile.PlanarReflectionMode=0
-bStreamSkeletalMeshLODs=(Default=False,PerPlatform=())
-bDiscardSkeletalMeshOptionalLODs=(Default=False,PerPlatform=())
-VisualizeCalibrationColorMaterialPath=None
-VisualizeCalibrationCustomMaterialPath=None
-VisualizeCalibrationGrayscaleMaterialPath=None
-r.Mobile.AntiAliasing=3
-r.Mobile.FloatPrecisionMode=2
-r.OpenGL.ForceDXC=0
 
-[/Script/HardwareTargeting.HardwareTargetingSettings]
-TargetedHardwareClass=Mobile
-AppliedTargetedHardwareClass=Mobile
-DefaultGraphicsPerformance=Scalable
-AppliedDefaultGraphicsPerformance=Scalable
+r.Shadow.Virtual.Enable=0
 
-[/Script/EngineSettings.GameMapsSettings]
-EditorStartupMap=/Game/VRTemplate/Maps/VRTemplateMap.VRTemplateMap
-LocalMapOptions=
-TransitionMap=None
-bUseSplitscreen=False
-TwoPlayerSplitscreenLayout=Horizontal
-ThreePlayerSplitscreenLayout=FavorTop
-FourPlayerSplitscreenLayout=Grid
-bOffsetPlayerGamepadIds=False
-GameInstanceClass=/Script/Engine.GameInstance
-GameDefaultMap=/Game/VRTemplate/Maps/VRTemplateMap.VRTemplateMap
-ServerDefaultMap=/Engine/Maps/Entry.Entry
-GlobalDefaultGameMode=/Game/VRTemplate/Blueprints/VRGameMode.VRGameMode_C
-GlobalDefaultServerGameMode=None
+r.DefaultFeature.AutoExposure.ExtendDefaultLuminanceRange=True
+r.DefaultFeature.AutoExposure.ExtendDefaultLuminanceRange=True
+
+r.DefaultFeature.LocalExposure.HighlightContrastScale=0.8
+r.DefaultFeature.LocalExposure.HighlightContrastScale=0.8
 
-[/Script/Slate.SlateSettings]
-bExplicitCanvasChildZOrder=True
+r.DefaultFeature.LocalExposure.ShadowContrastScale=0.8
+r.DefaultFeature.LocalExposure.ShadowContrastScale=0.8
+
+r.GenerateMeshDistanceFields=True
+
+r.DynamicGlobalIlluminationMethod=0
+
+r.ReflectionMethod=1
+r.Mobile.AntiAliasing=3
+r.AntiAliasingMethod=3
+
+[/Script/AndroidRuntimeSettings.AndroidRuntimeSettings]
+bBuildForES31=False
+ExtraApplicationSettings=<meta-data android:name="com.oculus.supportedDevices" android:value="quest|quest2|questpro" />
+bPackageForMetaQuest=True
 
 [/Script/WindowsTargetPlatform.WindowsTargetSettings]
-DefaultGraphicsRHI=DefaultGraphicsRHI_DX12
+DefaultGraphicsRHI=DefaultGraphicsRHI_Default
 -D3D12TargetedShaderFormats=PCD3D_SM5
 +D3D12TargetedShaderFormats=PCD3D_SM6
 -D3D11TargetedShaderFormats=PCD3D_SM5
 +D3D11TargetedShaderFormats=PCD3D_SM5
 Compiler=Default
-MinimumOSVersion=MSOS_Vista
-bTarget32Bit=False
 AudioSampleRate=48000
 AudioCallbackBufferFrameSize=1024
 AudioNumBuffersToEnqueue=1
 AudioMaxChannels=0
 AudioNumSourceWorkers=4
-SpatializationPlugin=Built-in Spatialization
-ReverbPlugin=Built-in Reverb
+SpatializationPlugin=
+SourceDataOverridePlugin=
+ReverbPlugin=
 OcclusionPlugin=
 CompressionOverrides=(bOverrideCompressionTimes=False,DurationThreshold=5.000000,MaxNumRandomBranches=0,SoundCueQualityIndex=0)
-bUseAudioStreamCaching=False
-CacheSizeKB=0
+CacheSizeKB=65536
 MaxChunkSizeOverrideKB=0
 bResampleForDevice=False
-MaxSampleRate=0.000000
-HighSampleRate=0.000000
-MedSampleRate=0.000000
-LowSampleRate=0.000000
-MinSampleRate=0.000000
-CompressionQualityModifier=0.000000
+MaxSampleRate=48000.000000
+HighSampleRate=32000.000000
+MedSampleRate=24000.000000
+LowSampleRate=12000.000000
+MinSampleRate=8000.000000
+CompressionQualityModifier=1.000000
 AutoStreamingThreshold=0.000000
 SoundCueCookQualityIndex=-1
 
-[/Script/AndroidRuntimeSettings.AndroidRuntimeSettings]
-PackageName=com.epicgames.[PROJECT]
-StoreVersion=1
-StoreVersionOffsetArmV7=0
-StoreVersionOffsetArm64=0
-StoreVersionOffsetX8664=0
-ApplicationDisplayName=
-VersionDisplayName=1.0
-MinSDKVersion=23
-TargetSDKVersion=25
-InstallLocation=InternalOnly
-bEnableGradle=True
-bEnableLint=False
-bPackageDataInsideApk=True
-bCreateAllPlatformsInstall=False
-bDisableVerifyOBBOnStartUp=False
-bForceSmallOBBFiles=False
-bAllowLargeOBBFiles=False
-bAllowPatchOBBFile=False
-bAllowOverflowOBBFiles=False
-bUseExternalFilesDir=False
-bPublicLogFiles=True
-Orientation=SensorLandscape
-MaxAspectRatio=2.100000
-bUseDisplayCutout=False
-bRestoreNotificationsOnReboot=False
-bFullScreen=True
-bEnableNewKeyboard=True
-DepthBufferPreference=Default
-bValidateTextureFormats=True
-bEnableBundle=False
-bEnableUniversalAPK=True
-bBundleABISplit=True
-bBundleLanguageSplit=True
-bBundleDensitySplit=True
-ExtraApplicationSettings=
-ExtraActivitySettings=
-bAndroidVoiceEnabled=False
-bPackageForMetaQuest=True
-bRemoveOSIG=False
-+GoogleVRCaps=Daydream33
-bGoogleVRSustainedPerformance=False
-KeyStore=UE.keystore
-KeyAlias=DistroKey
-KeyStorePassword=codesigner
-KeyPassword=
-bBuildForArm64=True
-bBuildForX8664=False
-bBuildForES31=False
-bSupportsVulkan=True
-bSupportsVulkanSM5=False
-ClangSanitizer=None
-bAndroidOpenGLSupportsBackbufferSampling=False
-bDetectVulkanByDefault=True
-bBuildWithHiddenSymbolVisibility=False
-bSaveSymbols=False
-bForceLDLinker=False
-bEnableGooglePlaySupport=False
-bUseGetAccounts=False
-GamesAppID=
-bEnableSnapshots=False
-bSupportAdMob=True
-AdMobAdUnitID=
-GooglePlayLicenseKey=
-GCMClientSenderID=
-bShowLaunchImage=True
-bAllowIMU=True
-bAllowControllers=True
-bBlockAndroidKeysOnControllers=False
-bControllersBlockDeviceFeedback=False
-AndroidAudio=Default
-AudioSampleRate=44100
-AudioCallbackBufferFrameSize=1024
-AudioNumBuffersToEnqueue=4
-AudioMaxChannels=0
-AudioNumSourceWorkers=0
-SpatializationPlugin=Built-in Spatialization
-ReverbPlugin=Built-in Reverb
-OcclusionPlugin=
-CompressionOverrides=(bOverrideCompressionTimes=False,DurationThreshold=5.000000,MaxNumRandomBranches=0,SoundCueQualityIndex=0)
-bUseAudioStreamCaching=False
-CacheSizeKB=0
-MaxChunkSizeOverrideKB=0
-bResampleForDevice=False
-SoundCueCookQualityIndex=-1
-MaxSampleRate=0.000000
-HighSampleRate=0.000000
-MedSampleRate=0.000000
-LowSampleRate=0.000000
-MinSampleRate=0.000000
-CompressionQualityModifier=0.000000
-AutoStreamingThreshold=0.000000
-AndroidGraphicsDebugger=None
-MaliGraphicsDebuggerPath=(Path="")
-bEnableMaliPerfCounters=False
-bMultiTargetFormat_ETC2=True
-bMultiTargetFormat_DXT=True
-bMultiTargetFormat_ASTC=True
-TextureFormatPriority_ETC2=0.200000
-TextureFormatPriority_DXT=0.600000
-TextureFormatPriority_ASTC=0.900000
-SDKAPILevelOverride=
-NDKAPILevelOverride=
-bStreamLandscapeMeshLODs=False
-
-[/Script/Engine.CollisionProfile]
--Profiles=(Name="NoCollision",CollisionEnabled=NoCollision,ObjectTypeName="WorldStatic",CustomResponses=((Channel="Visibility",Response=ECR_Ignore),(Channel="Camera",Response=ECR_Ignore)),HelpMessage="No collision",bCanModify=False)
--Profiles=(Name="BlockAll",CollisionEnabled=QueryAndPhysics,ObjectTypeName="WorldStatic",CustomResponses=,HelpMessage="WorldStatic object that blocks all actors by default. All new custom channels will use its own default response. ",bCanModify=False)
--Profiles=(Name="OverlapAll",CollisionEnabled=QueryOnly,ObjectTypeName="WorldStatic",CustomResponses=((Channel="WorldStatic",Response=ECR_Overlap),(Channel="Pawn",Response=ECR_Overlap),(Channel="Visibility",Response=ECR_Overlap),(Channel="WorldDynamic",Response=ECR_Overlap),(Channel="Camera",Response=ECR_Overlap),(Channel="PhysicsBody",Response=ECR_Overlap),(Channel="Vehicle",Response=ECR_Overlap),(Channel="Destructible",Response=ECR_Overlap)),HelpMessage="WorldStatic object that overlaps all actors by default. All new custom channels will use its own default response. ",bCanModify=False)
--Profiles=(Name="BlockAllDynamic",CollisionEnabled=QueryAndPhysics,ObjectTypeName="WorldDynamic",CustomResponses=,HelpMessage="WorldDynamic object that blocks all actors by default. All new custom channels will use its own default response. ",bCanModify=False)
--Profiles=(Name="OverlapAllDynamic",CollisionEnabled=QueryOnly,ObjectTypeName="WorldDynamic",CustomResponses=((Channel="WorldStatic",Response=ECR_Overlap),(Channel="Pawn",Response=ECR_Overlap),(Channel="Visibility",Response=ECR_Overlap),(Channel="WorldDynamic",Response=ECR_Overlap),(Channel="Camera",Response=ECR_Overlap),(Channel="PhysicsBody",Response=ECR_Overlap),(Channel="Vehicle",Response=ECR_Overlap),(Channel="Destructible",Response=ECR_Overlap)),HelpMessage="WorldDynamic object that overlaps all actors by default. All new custom channels will use its own default response. ",bCanModify=False)
--Profiles=(Name="IgnoreOnlyPawn",CollisionEnabled=QueryOnly,ObjectTypeName="WorldDynamic",CustomResponses=((Channel="Pawn",Response=ECR_Ignore),(Channel="Vehicle",Response=ECR_Ignore)),HelpMessage="WorldDynamic object that ignores Pawn and Vehicle. All other channels will be set to default.",bCanModify=False)
--Profiles=(Name="OverlapOnlyPawn",CollisionEnabled=QueryOnly,ObjectTypeName="WorldDynamic",CustomResponses=((Channel="Pawn",Response=ECR_Overlap),(Channel="Vehicle",Response=ECR_Overlap),(Channel="Camera",Response=ECR_Ignore)),HelpMessage="WorldDynamic object that overlaps Pawn, Camera, and Vehicle. All other channels will be set to default. ",bCanModify=False)
--Profiles=(Name="Pawn",CollisionEnabled=QueryAndPhysics,ObjectTypeName="Pawn",CustomResponses=((Channel="Visibility",Response=ECR_Ignore)),HelpMessage="Pawn object. Can be used for capsule of any playerable character or AI. ",bCanModify=False)
--Profiles=(Name="Spectator",CollisionEnabled=QueryOnly,ObjectTypeName="Pawn",CustomResponses=((Channel="WorldStatic",Response=ECR_Block),(Channel="Pawn",Response=ECR_Ignore),(Channel="Visibility",Response=ECR_Ignore),(Channel="WorldDynamic",Response=ECR_Ignore),(Channel="Camera",Response=ECR_Ignore),(Channel="PhysicsBody",Response=ECR_Ignore),(Channel="Vehicle",Response=ECR_Ignore),(Channel="Destructible",Response=ECR_Ignore)),HelpMessage="Pawn object that ignores all other actors except WorldStatic.",bCanModify=False)
--Profiles=(Name="CharacterMesh",CollisionEnabled=QueryOnly,ObjectTypeName="Pawn",CustomResponses=((Channel="Pawn",Response=ECR_Ignore),(Channel="Vehicle",Response=ECR_Ignore),(Channel="Visibility",Response=ECR_Ignore)),HelpMessage="Pawn object that is used for Character Mesh. All other channels will be set to default.",bCanModify=False)
--Profiles=(Name="PhysicsActor",CollisionEnabled=QueryAndPhysics,ObjectTypeName="PhysicsBody",CustomResponses=,HelpMessage="Simulating actors",bCanModify=False)
--Profiles=(Name="Destructible",CollisionEnabled=QueryAndPhysics,ObjectTypeName="Destructible",CustomResponses=,HelpMessage="Destructible actors",bCanModify=False)
--Profiles=(Name="InvisibleWall",CollisionEnabled=QueryAndPhysics,ObjectTypeName="WorldStatic",CustomResponses=((Channel="Visibility",Response=ECR_Ignore)),HelpMessage="WorldStatic object that is invisible.",bCanModify=False)
--Profiles=(Name="InvisibleWallDynamic",CollisionEnabled=QueryAndPhysics,ObjectTypeName="WorldDynamic",CustomResponses=((Channel="Visibility",Response=ECR_Ignore)),HelpMessage="WorldDynamic object that is invisible.",bCanModify=False)
--Profiles=(Name="Trigger",CollisionEnabled=QueryOnly,ObjectTypeName="WorldDynamic",CustomResponses=((Channel="WorldStatic",Response=ECR_Overlap),(Channel="Pawn",Response=ECR_Overlap),(Channel="Visibility",Response=ECR_Ignore),(Channel="WorldDynamic",Response=ECR_Overlap),(Channel="Camera",Response=ECR_Overlap),(Channel="PhysicsBody",Response=ECR_Overlap),(Channel="Vehicle",Response=ECR_Overlap),(Channel="Destructible",Response=ECR_Overlap)),HelpMessage="WorldDynamic object that is used for trigger. All other channels will be set to default.",bCanModify=False)
--Profiles=(Name="Ragdoll",CollisionEnabled=QueryAndPhysics,ObjectTypeName="PhysicsBody",CustomResponses=((Channel="Pawn",Response=ECR_Ignore),(Channel="Visibility",Response=ECR_Ignore)),HelpMessage="Simulating Skeletal Mesh Component. All other channels will be set to default.",bCanModify=False)
--Profiles=(Name="Vehicle",CollisionEnabled=QueryAndPhysics,ObjectTypeName="Vehicle",CustomResponses=,HelpMessage="Vehicle object that blocks Vehicle, WorldStatic, and WorldDynamic. All other channels will be set to default.",bCanModify=False)
--Profiles=(Name="UI",CollisionEnabled=QueryOnly,ObjectTypeName="WorldDynamic",CustomResponses=((Channel="WorldStatic",Response=ECR_Overlap),(Channel="Pawn",Response=ECR_Overlap),(Channel="Visibility",Response=ECR_Block),(Channel="WorldDynamic",Response=ECR_Overlap),(Channel="Camera",Response=ECR_Overlap),(Channel="PhysicsBody",Response=ECR_Overlap),(Channel="Vehicle",Response=ECR_Overlap),(Channel="Destructible",Response=ECR_Overlap)),HelpMessage="WorldStatic object that overlaps all actors by default. All new custom channels will use its own default response. ",bCanModify=False)
-+Profiles=(Name="NoCollision",CollisionEnabled=NoCollision,bCanModify=False,ObjectTypeName="WorldStatic",CustomResponses=((Channel="Visibility",Response=ECR_Ignore),(Channel="Camera",Response=ECR_Ignore)),HelpMessage="No collision")
-+Profiles=(Name="BlockAll",CollisionEnabled=QueryAndPhysics,bCanModify=False,ObjectTypeName="WorldStatic",CustomResponses=,HelpMessage="WorldStatic object that blocks all actors by default. All new custom channels will use its own default response. ")
-+Profiles=(Name="OverlapAll",CollisionEnabled=QueryOnly,bCanModify=False,ObjectTypeName="WorldStatic",CustomResponses=((Channel="WorldStatic",Response=ECR_Overlap),(Channel="Pawn",Response=ECR_Overlap),(Channel="Visibility",Response=ECR_Overlap),(Channel="WorldDynamic",Response=ECR_Overlap),(Channel="Camera",Response=ECR_Overlap),(Channel="PhysicsBody",Response=ECR_Overlap),(Channel="Vehicle",Response=ECR_Overlap),(Channel="Destructible",Response=ECR_Overlap)),HelpMessage="WorldStatic object that overlaps all actors by default. All new custom channels will use its own default response. ")
-+Profiles=(Name="BlockAllDynamic",CollisionEnabled=QueryAndPhysics,bCanModify=False,ObjectTypeName="WorldDynamic",CustomResponses=,HelpMessage="WorldDynamic object that blocks all actors by default. All new custom channels will use its own default response. ")
-+Profiles=(Name="OverlapAllDynamic",CollisionEnabled=QueryOnly,bCanModify=False,ObjectTypeName="WorldDynamic",CustomResponses=((Channel="WorldStatic",Response=ECR_Overlap),(Channel="Pawn",Response=ECR_Overlap),(Channel="Visibility",Response=ECR_Overlap),(Channel="WorldDynamic",Response=ECR_Overlap),(Channel="Camera",Response=ECR_Overlap),(Channel="PhysicsBody",Response=ECR_Overlap),(Channel="Vehicle",Response=ECR_Overlap),(Channel="Destructible",Response=ECR_Overlap)),HelpMessage="WorldDynamic object that overlaps all actors by default. All new custom channels will use its own default response. ")
-+Profiles=(Name="IgnoreOnlyPawn",CollisionEnabled=QueryOnly,bCanModify=False,ObjectTypeName="WorldDynamic",CustomResponses=((Channel="Pawn",Response=ECR_Ignore),(Channel="Vehicle",Response=ECR_Ignore)),HelpMessage="WorldDynamic object that ignores Pawn and Vehicle. All other channels will be set to default.")
-+Profiles=(Name="OverlapOnlyPawn",CollisionEnabled=QueryOnly,bCanModify=False,ObjectTypeName="WorldDynamic",CustomResponses=((Channel="Pawn",Response=ECR_Overlap),(Channel="Vehicle",Response=ECR_Overlap),(Channel="Camera",Response=ECR_Ignore)),HelpMessage="WorldDynamic object that overlaps Pawn, Camera, and Vehicle. All other channels will be set to default. ")
-+Profiles=(Name="Pawn",CollisionEnabled=QueryAndPhysics,bCanModify=False,ObjectTypeName="Pawn",CustomResponses=((Channel="Visibility",Response=ECR_Ignore)),HelpMessage="Pawn object. Can be used for capsule of any playerable character or AI. ")
-+Profiles=(Name="Spectator",CollisionEnabled=QueryOnly,bCanModify=False,ObjectTypeName="Pawn",CustomResponses=((Channel="WorldStatic"),(Channel="Pawn",Response=ECR_Ignore),(Channel="Visibility",Response=ECR_Ignore),(Channel="WorldDynamic",Response=ECR_Ignore),(Channel="Camera",Response=ECR_Ignore),(Channel="PhysicsBody",Response=ECR_Ignore),(Channel="Vehicle",Response=ECR_Ignore),(Channel="Destructible",Response=ECR_Ignore)),HelpMessage="Pawn object that ignores all other actors except WorldStatic.")
-+Profiles=(Name="CharacterMesh",CollisionEnabled=QueryOnly,bCanModify=False,ObjectTypeName="Pawn",CustomResponses=((Channel="Pawn",Response=ECR_Ignore),(Channel="Vehicle",Response=ECR_Ignore),(Channel="Visibility",Response=ECR_Ignore)),HelpMessage="Pawn object that is used for Character Mesh. All other channels will be set to default.")
-+Profiles=(Name="PhysicsActor",CollisionEnabled=QueryAndPhysics,bCanModify=False,ObjectTypeName="PhysicsBody",CustomResponses=,HelpMessage="Simulating actors")
-+Profiles=(Name="Destructible",CollisionEnabled=QueryAndPhysics,bCanModify=False,ObjectTypeName="Destructible",CustomResponses=,HelpMessage="Destructible actors")
-+Profiles=(Name="InvisibleWall",CollisionEnabled=QueryAndPhysics,bCanModify=False,ObjectTypeName="WorldStatic",CustomResponses=((Channel="Visibility",Response=ECR_Ignore)),HelpMessage="WorldStatic object that is invisible.")
-+Profiles=(Name="InvisibleWallDynamic",CollisionEnabled=QueryAndPhysics,bCanModify=False,ObjectTypeName="WorldDynamic",CustomResponses=((Channel="Visibility",Response=ECR_Ignore)),HelpMessage="WorldDynamic object that is invisible.")
-+Profiles=(Name="Trigger",CollisionEnabled=QueryOnly,bCanModify=False,ObjectTypeName="WorldDynamic",CustomResponses=((Channel="WorldStatic",Response=ECR_Overlap),(Channel="Pawn",Response=ECR_Overlap),(Channel="Visibility",Response=ECR_Ignore),(Channel="WorldDynamic",Response=ECR_Overlap),(Channel="Camera",Response=ECR_Overlap),(Channel="PhysicsBody",Response=ECR_Overlap),(Channel="Vehicle",Response=ECR_Overlap),(Channel="Destructible",Response=ECR_Overlap)),HelpMessage="WorldDynamic object that is used for trigger. All other channels will be set to default.")
-+Profiles=(Name="Ragdoll",CollisionEnabled=QueryAndPhysics,bCanModify=False,ObjectTypeName="PhysicsBody",CustomResponses=((Channel="Pawn",Response=ECR_Ignore),(Channel="Visibility",Response=ECR_Ignore)),HelpMessage="Simulating Skeletal Mesh Component. All other channels will be set to default.")
-+Profiles=(Name="Vehicle",CollisionEnabled=QueryAndPhysics,bCanModify=False,ObjectTypeName="Vehicle",CustomResponses=,HelpMessage="Vehicle object that blocks Vehicle, WorldStatic, and WorldDynamic. All other channels will be set to default.")
-+Profiles=(Name="UI",CollisionEnabled=QueryOnly,bCanModify=False,ObjectTypeName="WorldDynamic",CustomResponses=((Channel="WorldStatic",Response=ECR_Overlap),(Channel="Pawn",Response=ECR_Overlap),(Channel="Visibility"),(Channel="WorldDynamic",Response=ECR_Overlap),(Channel="Camera",Response=ECR_Overlap),(Channel="PhysicsBody",Response=ECR_Overlap),(Channel="Vehicle",Response=ECR_Overlap),(Channel="Destructible",Response=ECR_Overlap)),HelpMessage="WorldStatic object that overlaps all actors by default. All new custom channels will use its own default response. ")
-+DefaultChannelResponses=(Channel=ECC_GameTraceChannel1,DefaultResponse=ECR_Ignore,bTraceType=True,bStaticObject=False,Name="3DWidget")
--ProfileRedirects=(OldName="BlockingVolume",NewName="InvisibleWall")
--ProfileRedirects=(OldName="InterpActor",NewName="IgnoreOnlyPawn")
--ProfileRedirects=(OldName="StaticMeshComponent",NewName="BlockAllDynamic")
--ProfileRedirects=(OldName="SkeletalMeshActor",NewName="PhysicsActor")
--ProfileRedirects=(OldName="InvisibleActor",NewName="InvisibleWallDynamic")
-+ProfileRedirects=(OldName="BlockingVolume",NewName="InvisibleWall")
-+ProfileRedirects=(OldName="InterpActor",NewName="IgnoreOnlyPawn")
-+ProfileRedirects=(OldName="StaticMeshComponent",NewName="BlockAllDynamic")
-+ProfileRedirects=(OldName="SkeletalMeshActor",NewName="PhysicsActor")
-+ProfileRedirects=(OldName="InvisibleActor",NewName="InvisibleWallDynamic")
--CollisionChannelRedirects=(OldName="Static",NewName="WorldStatic")
--CollisionChannelRedirects=(OldName="Dynamic",NewName="WorldDynamic")
--CollisionChannelRedirects=(OldName="VehicleMovement",NewName="Vehicle")
--CollisionChannelRedirects=(OldName="PawnMovement",NewName="Pawn")
-+CollisionChannelRedirects=(OldName="Static",NewName="WorldStatic")
-+CollisionChannelRedirects=(OldName="Dynamic",NewName="WorldDynamic")
-+CollisionChannelRedirects=(OldName="VehicleMovement",NewName="Vehicle")
-+CollisionChannelRedirects=(OldName="PawnMovement",NewName="Pawn")
-
 [/Script/NavigationSystem.NavigationSystemV1]
 bAllowClientSideNavigation=True
\ No newline at end of file
