XR | VR Template | Adding Quest3 to AndroidRuntimeSettings->ExtraApplicationSettings. Without this, Quest 3 will run in Quest 2 compatibility mode.

diff --git a/Templates/TP_VirtualRealityBP/Config/DefaultEngine.ini b/Templates/TP_VirtualRealityBP/Config/DefaultEngine.ini
--- a/Templates/TP_VirtualRealityBP/Config/DefaultEngine.ini
+++ b/Templates/TP_VirtualRealityBP/Config/DefaultEngine.ini
@@ -35,7 +35,7 @@ r.AntiAliasingMethod=3
 
 [/Script/AndroidRuntimeSettings.AndroidRuntimeSettings]
 bBuildForES31=False
-ExtraApplicationSettings=<meta-data android:name="com.oculus.supportedDevices" android:value="quest|quest2|questpro" />
+ExtraApplicationSettings=<meta-data android:name="com.oculus.supportedDevices" android:value="quest|quest2|questpro|quest3" />
 bPackageForMetaQuest=True
 
 [/Script/WindowsTargetPlatform.WindowsTargetSettings]
@@ -73,3 +73,4 @@ bAllowClientSideNavigation=True
 [/Script/LinuxTargetPlatform.LinuxTargetSettings]
 -TargetedRHIs=SF_VULKAN_SM5
 +TargetedRHIs=SF_VULKAN_SM6
+
