VR Template | Setting Allow Static Lighting to True as it's now set to false in editor by default.
Attempt #2.

diff --git a/Templates/TP_VirtualRealityBP/Config/DefaultEngine.ini b/Templates/TP_VirtualRealityBP/Config/DefaultEngine.ini
--- a/Templates/TP_VirtualRealityBP/Config/DefaultEngine.ini
+++ b/Templates/TP_VirtualRealityBP/Config/DefaultEngine.ini
@@ -32,6 +32,7 @@ r.DynamicGlobalIlluminationMethod=0
 r.ReflectionMethod=1
 r.Mobile.AntiAliasing=3
 r.AntiAliasingMethod=3
+r.AllowStaticLighting=True
 
 [/Script/AndroidRuntimeSettings.AndroidRuntimeSettings]
 bBuildForES31=False
