VR Template | Updated Android Min and Target SDK to 32

diff --git a/Templates/TP_VirtualRealityBP/Config/DefaultEngine.ini b/Templates/TP_VirtualRealityBP/Config/DefaultEngine.ini
--- a/Templates/TP_VirtualRealityBP/Config/DefaultEngine.ini
+++ b/Templates/TP_VirtualRealityBP/Config/DefaultEngine.ini
@@ -38,6 +38,8 @@ r.AllowStaticLighting=True
 bBuildForES31=False
 ExtraApplicationSettings=<meta-data android:name="com.oculus.supportedDevices" android:value="quest2|questpro|quest3" />
 bPackageForMetaQuest=True
+MinSDKVersion=32
+TargetSDKVersion=32
 
 [/Script/WindowsTargetPlatform.WindowsTargetSettings]
 DefaultGraphicsRHI=DefaultGraphicsRHI_Default
