VR Template | The mannequin hands are now a shared template resource - this enable us to use them in multiple XR templates

#virtualized

diff --git a/Templates/TP_VirtualRealityBP/Config/TemplateDefs.ini b/Templates/TP_VirtualRealityBP/Config/TemplateDefs.ini
--- a/Templates/TP_VirtualRealityBP/Config/TemplateDefs.ini
+++ b/Templates/TP_VirtualRealityBP/Config/TemplateDefs.ini
@@ -45,6 +45,7 @@ ReplacementsInFiles=(Extensions=("cpp","h","ini","cs"),From="%TEMPLATENAME%",To=
 SortKey=VirtualRealityBP
 SharedContentPacks=(MountName="Weapons",DetailLevels=("Standard"))
 SharedContentPacks=(MountName="VRSpectator",DetailLevels=("Standard"))
+SharedContentPacks=(MountName="XRMannequins",DetailLevels=("Standard"))
 SharedContentPacks=(MountName="LevelPrototyping",DetailLevels=("High"))
 EditDetailLevelPreference="High"
 FocusAsset=/Game/VRTemplate/Blueprints
diff --git a/Templates/TP_VirtualRealityBP/manifest.json b/Templates/TP_VirtualRealityBP/manifest.json
--- a/Templates/TP_VirtualRealityBP/manifest.json
+++ b/Templates/TP_VirtualRealityBP/manifest.json
@@ -86,6 +86,12 @@
 				"Standard"
 			]
 		},
+		{
+			"MountName": "XRMannequins",
+			"DetailLevels": [
+				"Standard"
+			]
+		},
 		{
 			"MountName": "VRSpectator",
 			"DetailLevels": [
