VR Template | VR Template content has now been moved to be a shared template resource named XRFramework. This enables the VR template content to be used in multiple templates.

#virtualized

diff --git a/Templates/TP_VirtualRealityBP/Config/DefaultEngine.ini b/Templates/TP_VirtualRealityBP/Config/DefaultEngine.ini
--- a/Templates/TP_VirtualRealityBP/Config/DefaultEngine.ini
+++ b/Templates/TP_VirtualRealityBP/Config/DefaultEngine.ini
@@ -1,7 +1,7 @@
 [/Script/EngineSettings.GameMapsSettings]
-GlobalDefaultGameMode=/Game/VRTemplate/Blueprints/VRGameMode.VRGameMode_C
-EditorStartupMap=/Game/VRTemplate/Maps/VRTemplateMap.VRTemplateMap
-GameDefaultMap=/Game/VRTemplate/Maps/VRTemplateMap.VRTemplateMap
+GlobalDefaultGameMode=/Game/XRFramework/Blueprints/BP_XRGameMode.BP_XRGameMode_C
+EditorStartupMap=/Game/VRTemplate/VRTemplateMap.VRTemplateMap
+GameDefaultMap=/Game/VRTemplate/VRTemplateMap.VRTemplateMap
 
 [/Script/Engine.RendererSettings]
 r.ForwardShading=True
@@ -134,3 +134,17 @@ bAllowClientSideNavigation=True
 +CollisionChannelRedirects=(OldName="VehicleMovement",NewName="Vehicle")
 +CollisionChannelRedirects=(OldName="PawnMovement",NewName="Pawn")
 
+[/Script/AndroidFileServerEditor.AndroidFileServerRuntimeSettings]
+bEnablePlugin=True
+bAllowNetworkConnection=True
+SecurityToken=208133E84DD6073D8265129FE6833986
+bIncludeInShipping=False
+bAllowExternalStartInShipping=False
+bCompileAFSProject=False
+bUseCompression=False
+bLogFiles=False
+bReportStats=False
+ConnectionType=USBOnly
+bUseManualIPAddress=False
+ManualIPAddress=
+
diff --git a/Templates/TP_VirtualRealityBP/Config/DefaultInput.ini b/Templates/TP_VirtualRealityBP/Config/DefaultInput.ini
--- a/Templates/TP_VirtualRealityBP/Config/DefaultInput.ini
+++ b/Templates/TP_VirtualRealityBP/Config/DefaultInput.ini
@@ -108,9 +108,21 @@ DefaultTouchInterface=None
 +ConsoleKeys=Tilde
 
 [/Script/EnhancedInput.EnhancedInputDeveloperSettings]
-+DefaultMappingContexts=(InputMappingContext="/Game/VRTemplate/Input/IMC_Default.IMC_Default",Priority=0,bAddImmediately=True,bRegisterWithUserSettings=False)
-+DefaultMappingContexts=(InputMappingContext="/Game/VRTemplate/Input/IMC_Hands.IMC_Hands",Priority=0,bAddImmediately=True,bRegisterWithUserSettings=False)
-+DefaultMappingContexts=(InputMappingContext="/Game/VRTemplate/Input/IMC_Menu.IMC_Menu",Priority=0,bAddImmediately=True,bRegisterWithUserSettings=False)
-+DefaultMappingContexts=(InputMappingContext="/Game/VRTemplate/Input/IMC_Weapon_Left.IMC_Weapon_Left",Priority=0,bAddImmediately=True,bRegisterWithUserSettings=False)
-+DefaultMappingContexts=(InputMappingContext="/Game/VRTemplate/Input/IMC_Weapon_Right.IMC_Weapon_Right",Priority=0,bAddImmediately=True,bRegisterWithUserSettings=False)
++DefaultMappingContexts=(InputMappingContext="/Game/XRFramework/Input/IMC_Default.IMC_Default",Priority=0,bAddImmediately=True,bRegisterWithUserSettings=False)
++DefaultMappingContexts=(InputMappingContext="/Game/XRFramework/Input/IMC_Hands.IMC_Hands",Priority=0,bAddImmediately=True,bRegisterWithUserSettings=False)
++DefaultMappingContexts=(InputMappingContext="/Game/XRFramework/Input/IMC_Weapon_Left.IMC_Weapon_Left",Priority=0,bAddImmediately=True,bRegisterWithUserSettings=False)
++DefaultMappingContexts=(InputMappingContext="/Game/XRFramework/Input/IMC_Weapon_Right.IMC_Weapon_Right",Priority=0,bAddImmediately=True,bRegisterWithUserSettings=False)
++DefaultMappingContexts=(InputMappingContext="/Game/XRFramework/Input/IMC_Menu.IMC_Menu",Priority=0,bAddImmediately=True,bRegisterWithUserSettings=False)
+UserSettingsClass=/Script/EnhancedInput.EnhancedInputUserSettings
+DefaultPlayerMappableKeyProfileClass=/Script/EnhancedInput.EnhancedPlayerMappableKeyProfile
+InputSettingsSaveSlotName=EnhancedInputUserSettings
+DefaultWorldInputClass=/Script/EnhancedInput.EnhancedPlayerInput
+bSendTriggeredEventsWhenInputIsFlushed=True
+bEnableUserSettings=False
+EnhancedInput.EnableDefaultMappingContexts=True
+EnhancedInput.OnlyTriggerLastActionInChord=True
+bEnableInputModeFiltering=True
+bEnableWorldSubsystem=False
+DefaultMappingContextInputModeQuery=(TokenStreamVersion=0,TagDictionary=((TagName="EnhancedInput.Modes.Default")),QueryTokenStream=(0,1,2,1,0),UserDescription="",AutoDescription="")
+DefaultInputMode=(GameplayTags=((TagName="EnhancedInput.Modes.Default")))
 
diff --git a/Templates/TP_VirtualRealityBP/Config/TemplateDefs.ini b/Templates/TP_VirtualRealityBP/Config/TemplateDefs.ini
--- a/Templates/TP_VirtualRealityBP/Config/TemplateDefs.ini
+++ b/Templates/TP_VirtualRealityBP/Config/TemplateDefs.ini
@@ -46,6 +46,7 @@ SortKey=VirtualRealityBP
 SharedContentPacks=(MountName="Weapons",DetailLevels=("Standard"))
 SharedContentPacks=(MountName="VRSpectator",DetailLevels=("Standard"))
 SharedContentPacks=(MountName="XRMannequins",DetailLevels=("Standard"))
+SharedContentPacks=(MountName="XRFramework",DetailLevels=("Standard"))
 SharedContentPacks=(MountName="LevelPrototyping",DetailLevels=("High"))
 EditDetailLevelPreference="High"
 FocusAsset=/Game/VRTemplate/Blueprints
diff --git a/Templates/TP_VirtualRealityBP/TP_VirtualRealityBP.uproject b/Templates/TP_VirtualRealityBP/TP_VirtualRealityBP.uproject
--- a/Templates/TP_VirtualRealityBP/TP_VirtualRealityBP.uproject
+++ b/Templates/TP_VirtualRealityBP/TP_VirtualRealityBP.uproject
@@ -10,8 +10,9 @@
 			"SupportedTargetPlatforms": [
 				"Win64",
 				"Linux",
+				"LinuxArm64",
 				"Android",
-				"VisionOS"
+				"PS5"
 			]
 		},
 		{
@@ -20,6 +21,7 @@
 			"SupportedTargetPlatforms": [
 				"Win64",
 				"Linux",
+				"LinuxArm64",
 				"Android"
 			]
 		},
@@ -29,8 +31,9 @@
 			"SupportedTargetPlatforms": [
 				"Win64",
 				"Linux",
+				"LinuxArm64",
 				"Android",
-				"VisionOS"
+				"PS5"
 			]
 		}
 	],
diff --git a/Templates/TP_VirtualRealityBP/manifest.json b/Templates/TP_VirtualRealityBP/manifest.json
--- a/Templates/TP_VirtualRealityBP/manifest.json
+++ b/Templates/TP_VirtualRealityBP/manifest.json
@@ -92,6 +92,12 @@
 				"Standard"
 			]
 		},
+		{
+			"MountName": "XRFramework",
+			"DetailLevels": [
+				"Standard"
+			]
+		},
 		{
 			"MountName": "VRSpectator",
 			"DetailLevels": [
diff --git a/Templates/TemplateResources/Standard/XRFramework/FeaturePack/manifest.json b/Templates/TemplateResources/Standard/XRFramework/FeaturePack/manifest.json
new file mode 100644
--- /dev/null
+++ b/Templates/TemplateResources/Standard/XRFramework/FeaturePack/manifest.json
@@ -0,0 +1,41 @@
+{
+	"Version": 1,
+	"Name": [
+		{
+			"Language": "en",
+			"Text": "XR Framework"
+		}
+	],
+	"Description": [
+		{
+			"Language": "en",
+			"Text": "Feature Pack for the XR Templates"
+		}
+	],
+	"AssetTypes": [
+		{
+			"Language": "en",
+			"Text": ""
+		}
+	],
+	"SearchTags": [
+		{
+			"Language": "en",
+			"Text": ""
+		}
+	],
+	"ClassTypes": " ",
+	"Category": "SharedPack",
+	"Thumbnail": "XRFramework.png",
+	"Screenshots": [
+		"XRFramework_Preview.png"
+	],
+	"Ident": "XRFrameworkStandard",
+	"AdditionalFiles": {
+		"DestinationFilesFolder": "XRFramework",
+		"AdditionalFilesList": [
+			"Templates/TemplateResources/Standard/XRFramework/Content/*.*"
+		]
+	}
+}
+
diff --git a/Templates/TemplateResources/Standard/XRFramework/contents.txt b/Templates/TemplateResources/Standard/XRFramework/contents.txt
new file mode 100644
--- /dev/null
+++ b/Templates/TemplateResources/Standard/XRFramework/contents.txt
@@ -0,0 +1,3 @@
+../../../Templates/TemplateResources/Standard/XRFramework/FeaturePack/manifest.json
+../../../Templates/TemplateResources/Standard/XRFramework/Media/XRFramework.png
+../../../Templates/TemplateResources/Standard/XRFramework/Media/XRFramework_Preview.png
\ No newline at end of file
