VR Template | Implemented an example of how Passthrough can be toggled. The option can be found in the hand menu. Note that this native Passthrough solution is currently only available on Meta Quest as the OpenXR runtime must support XR_EXT_composition_layer_inverted_alpha.

#virtualized

diff --git a/Templates/TP_VirtualRealityBP/Config/DefaultEngine.ini b/Templates/TP_VirtualRealityBP/Config/DefaultEngine.ini
--- a/Templates/TP_VirtualRealityBP/Config/DefaultEngine.ini
+++ b/Templates/TP_VirtualRealityBP/Config/DefaultEngine.ini
@@ -33,6 +33,7 @@ r.ReflectionMethod=1
 r.Mobile.AntiAliasing=3
 r.AntiAliasingMethod=3
 r.AllowStaticLighting=True
+r.Mobile.PropagateAlpha=True
 
 [/Script/AndroidRuntimeSettings.AndroidRuntimeSettings]
 bBuildForES31=False
@@ -148,3 +149,6 @@ ConnectionType=USBOnly
 bUseManualIPAddress=False
 ManualIPAddress=
 
+[/Script/OpenXRHMD.OpenXRHMDSettings]
+xr.OpenXRInvertAlpha=True
+
