VR Template | Removing Android File Server settings from DefaultEngine.ini

diff --git a/Templates/TP_VirtualRealityBP/Config/DefaultEngine.ini b/Templates/TP_VirtualRealityBP/Config/DefaultEngine.ini
--- a/Templates/TP_VirtualRealityBP/Config/DefaultEngine.ini
+++ b/Templates/TP_VirtualRealityBP/Config/DefaultEngine.ini
@@ -135,20 +135,6 @@ bAllowClientSideNavigation=True
 +CollisionChannelRedirects=(OldName="VehicleMovement",NewName="Vehicle")
 +CollisionChannelRedirects=(OldName="PawnMovement",NewName="Pawn")
 
-[/Script/AndroidFileServerEditor.AndroidFileServerRuntimeSettings]
-bEnablePlugin=True
-bAllowNetworkConnection=True
-SecurityToken=208133E84DD6073D8265129FE6833986
-bIncludeInShipping=False
-bAllowExternalStartInShipping=False
-bCompileAFSProject=False
-bUseCompression=False
-bLogFiles=False
-bReportStats=False
-ConnectionType=USBOnly
-bUseManualIPAddress=False
-ManualIPAddress=
-
 [/Script/OpenXRHMD.OpenXRHMDSettings]
 xr.OpenXRInvertAlpha=True
 
