Mobile: Set forward shading as a default for VR template

diff --git a/Templates/TP_VirtualRealityBP/Config/DefaultEngine.ini b/Templates/TP_VirtualRealityBP/Config/DefaultEngine.ini
--- a/Templates/TP_VirtualRealityBP/Config/DefaultEngine.ini
+++ b/Templates/TP_VirtualRealityBP/Config/DefaultEngine.ini
@@ -34,6 +34,7 @@ r.Mobile.AntiAliasing=3
 r.AntiAliasingMethod=3
 r.AllowStaticLighting=True
 r.Mobile.PropagateAlpha=True
+r.Mobile.ShadingPath=0
 
 [/Script/AndroidRuntimeSettings.AndroidRuntimeSettings]
 bBuildForES31=False
