Disable distance fields for VR Template.
Distance fields were enabled alongside the swap to using multi-pass deferred as the default mobile renderer. On the forward mobile renderer used by VR Template, distance fields require a full depth pre-pass, causing significant performance regressions.

diff --git a/Templates/TP_VirtualRealityBP/Config/Android/AndroidEngine.ini b/Templates/TP_VirtualRealityBP/Config/Android/AndroidEngine.ini
new file mode 100644
--- /dev/null
+++ b/Templates/TP_VirtualRealityBP/Config/Android/AndroidEngine.ini
@@ -0,0 +1,2 @@
+[/Script/Engine.RendererSettings]
+r.DistanceFields=0
\ No newline at end of file
