Temporarily disable Substrate on VR Template due to crashes at launch.

diff --git a/Templates/TP_VirtualRealityBP/Config/DefaultEngine.ini b/Templates/TP_VirtualRealityBP/Config/DefaultEngine.ini
--- a/Templates/TP_VirtualRealityBP/Config/DefaultEngine.ini
+++ b/Templates/TP_VirtualRealityBP/Config/DefaultEngine.ini
@@ -36,6 +36,8 @@ r.AllowStaticLighting=True
 r.Mobile.PropagateAlpha=True
 r.Mobile.ShadingPath=0
 
+r.Substrate=False
+
 [/Script/AndroidRuntimeSettings.AndroidRuntimeSettings]
 bBuildForES31=False
 ExtraApplicationSettings=<meta-data android:name="com.oculus.supportedDevices" android:value="quest2|questpro|quest3" />
