Re-enable substrate on VR Template now that the FStochasticLightingTileClassificationMarkCS crash has been fixed in CL 46629309.

diff --git a/Templates/TP_VirtualRealityBP/Config/DefaultEngine.ini b/Templates/TP_VirtualRealityBP/Config/DefaultEngine.ini
--- a/Templates/TP_VirtualRealityBP/Config/DefaultEngine.ini
+++ b/Templates/TP_VirtualRealityBP/Config/DefaultEngine.ini
@@ -36,8 +36,6 @@ r.AllowStaticLighting=True
 r.Mobile.PropagateAlpha=True
 r.Mobile.ShadingPath=0
 
-r.Substrate=False
-
 [/Script/AndroidRuntimeSettings.AndroidRuntimeSettings]
 bBuildForES31=False
 ExtraApplicationSettings=<meta-data android:name="com.oculus.supportedDevices" android:value="quest2|questpro|quest3" />
