VisionOS MSAA break means that the vr template does not work out of the box.
-Changed template inis so that we use FXAA rather than MSAA by default.
-Also some mild improvements to the ini setup to make it easier to change settings.

diff --git a/Templates/TP_VirtualRealityBP/Config/VisionOS/VisionOSDeviceProfiles.ini b/Templates/TP_VirtualRealityBP/Config/VisionOS/VisionOSDeviceProfiles.ini
new file mode 100644
--- /dev/null
+++ b/Templates/TP_VirtualRealityBP/Config/VisionOS/VisionOSDeviceProfiles.ini
@@ -0,0 +1,6 @@
+﻿[VisionPro DeviceProfile]
+;MSAA
+;+CVars=r.Mobile.XRMSAAMode=2 ;Note: this is in BaseDeviceProfiles.ini, so you don't actually need to set it.
+
+;NOT MSAA
++CVars=r.Mobile.XRMSAAMode=0
diff --git a/Templates/TP_VirtualRealityBP/Config/VisionOS/VisionOSEngine.ini b/Templates/TP_VirtualRealityBP/Config/VisionOS/VisionOSEngine.ini
--- a/Templates/TP_VirtualRealityBP/Config/VisionOS/VisionOSEngine.ini
+++ b/Templates/TP_VirtualRealityBP/Config/VisionOS/VisionOSEngine.ini
@@ -1,12 +1,30 @@
-
 [/Script/Engine.RendererSettings]
-// These settings are currently required for all visionOS projects.
+;These settings are currently required for all visionOS projects.
 vr.InstancedStereo=False
 vr.MobileMultiView=False
 xr.OpenXRAcquireMode=1
 
-//Enable these settings for .mixed immesion mode, aka passthrough, AR.
-//You will also need to edit UESwift.swift
-//r.Mobile.PropagateAlpha=1
-//r.PostProcessing.PropagateAlpha=1
-//r.AlphaInvertPass=1
\ No newline at end of file
+;MSAA ;Note these are the default settings, specified in DefaultEngine.ini and BaseDeviceProfiles.ini
+;r.Mobile.XRMSAAMode=2
+;r.Mobile.AntiAliasing=3
+;r.AntiAliasingMethod=3
+;MSAA END
+
+;FXAA
+;r.Mobile.XRMSAAMode=0 ; This is overridden in BaseDeviceProfile.ini for VisionPro.  It can be changed in VisionOSDeviceProfiles.ini for this project.
+r.Mobile.AntiAliasing=1
+r.AntiAliasingMethod=1
+;FXAA END
+
+;MIXED IMMERSION Enable these settings for .mixed immesion mode, aka passthrough, AR.
+;Comment them all out and full immesion will be used.
+;Note: You would also probably need to delete the skybox from the level, etc, so that you are not covering up all of the real world.
+r.Mobile.PropagateAlpha=1
+r.PostProcessing.PropagateAlpha=1
+r.AlphaInvertPass=1
+
+;Note: this could also be done in Editor->ProjectSettings->Platforms->visionOS->Immersive Space->ImmsersiveStyle to 1.
+[/Script/VisionOSRuntimeSettings.VisionOSRuntimeSettings]
+ImmersiveStyle=1
+;MIXED IMMERSION END
+
