Update VR Template to set r.Mobile.AllowFramebufferFetch=1, which causes us to use single-pass/Vulkan subpass deferred rendering when using the mobile deferred renderer.
This was the existing behavior before multi-pass deferred rendering was made default in 5.7. Multi-pass deferred rendering is currently broken in the right eye when mobile multi-view is enabled.

diff --git a/Templates/TP_VirtualRealityBP/Config/Android/AndroidEngine.ini b/Templates/TP_VirtualRealityBP/Config/Android/AndroidEngine.ini
--- a/Templates/TP_VirtualRealityBP/Config/Android/AndroidEngine.ini
+++ b/Templates/TP_VirtualRealityBP/Config/Android/AndroidEngine.ini
@@ -1,2 +1,6 @@
 [/Script/Engine.RendererSettings]
-r.DistanceFields=0
\ No newline at end of file
+r.DistanceFields=0
+
+; r.Mobile.AllowFramebufferFetch=1 reverts the mobile deferred renderer from multi-pass to single-pass/Vulkan subpasses, which was the default before UE 5.7.
+; Multi-pass deferred is disabled for VR Template for 5.7.1 due to rendering issues with mobile multi-view.
+r.Mobile.AllowFramebufferFetch=1
\ No newline at end of file
