Update VR Template to set r.Mobile.AllowFramebufferFetch=1, which causes us to use single-pass/Vulkan subpass deferred rendering when using the mobile deferred renderer.
This was the existing behavior before multi-pass deferred rendering was made default in 5.7. Multi-pass deferred rendering is currently broken in the right eye when mobile multi-view is enabled.

diff --git a/Templates/TP_VirtualRealityBP/Config/DefaultEngine.ini b/Templates/TP_VirtualRealityBP/Config/DefaultEngine.ini
--- a/Templates/TP_VirtualRealityBP/Config/DefaultEngine.ini
+++ b/Templates/TP_VirtualRealityBP/Config/DefaultEngine.ini
@@ -36,6 +36,9 @@ r.AllowStaticLighting=True
 r.Mobile.PropagateAlpha=True
 r.Mobile.ShadingPath=0
 
+; Multi-pass mobile deferred does not currently support mobile multi-view
+r.Mobile.AllowFramebufferFetch=1
+
 [/Script/AndroidRuntimeSettings.AndroidRuntimeSettings]
 bBuildForES31=False
 ExtraApplicationSettings=<meta-data android:name="com.oculus.supportedDevices" android:value="quest2|questpro|quest3" />
