XR | Template | The VR Template map has had the lighting separated into it's own level. This enables sharing of the content in the map between separate template configurations.

#virtualized
#Synced-CL 51379069

diff --git a/Templates/TP_VirtualRealityBP/Config/DefaultEngine.ini b/Templates/TP_VirtualRealityBP/Config/DefaultEngine.ini
--- a/Templates/TP_VirtualRealityBP/Config/DefaultEngine.ini
+++ b/Templates/TP_VirtualRealityBP/Config/DefaultEngine.ini
@@ -1,7 +1,7 @@
 [/Script/EngineSettings.GameMapsSettings]
 GlobalDefaultGameMode=/Game/XRFramework/Blueprints/BP_XRGameMode.BP_XRGameMode_C
-EditorStartupMap=/Game/VRTemplate/VRTemplateMap.VRTemplateMap
-GameDefaultMap=/Game/VRTemplate/VRTemplateMap.VRTemplateMap
+EditorStartupMap=/Game/XRFramework/Levels/L_XRTemplate.L_XRTemplate
+GameDefaultMap=/Game/XRFramework/Levels/L_XRTemplate.L_XRTemplate
 
 [/Script/Engine.RendererSettings]
 r.ForwardShading=True
@@ -142,3 +142,17 @@ bAllowClientSideNavigation=True
 [/Script/OpenXRHMD.OpenXRHMDSettings]
 xr.OpenXRInvertAlpha=True
 
+[/Script/AndroidFileServerEditor.AndroidFileServerRuntimeSettings]
+bEnablePlugin=True
+bAllowNetworkConnection=True
+SecurityToken=28DADB7345C4867967CFD59E4A1A48A2
+bIncludeInShipping=False
+bAllowExternalStartInShipping=False
+bCompileAFSProject=False
+bUseCompression=False
+bLogFiles=False
+bReportStats=False
+ConnectionType=USBOnly
+bUseManualIPAddress=False
+ManualIPAddress=
+
