Disable PSO precaching for Android devices on VR Template.
Precaching is currently slow due to both the limited speed of standalone Android-based XR devices, and VR Template precaching more PSOs than other templates. Disabling precaching significantly improves iteration times by avoiding a 20-30 second black screen on launch. Standalone XR apps are recommended to use Bundled PSO Caches, or at least conceal the the initial precaching with a loading screen.

#Synced-CL 52062538

diff --git a/Templates/TP_VirtualRealityBP/Config/Android/AndroidEngine.ini b/Templates/TP_VirtualRealityBP/Config/Android/AndroidEngine.ini
--- a/Templates/TP_VirtualRealityBP/Config/Android/AndroidEngine.ini
+++ b/Templates/TP_VirtualRealityBP/Config/Android/AndroidEngine.ini
@@ -1,6 +1,12 @@
 [/Script/Engine.RendererSettings]
+; Disabling distance fields avoids requiring a full depth pre-pass, which is slow on standalone XR devices.
 r.DistanceFields=0
 
-; r.Mobile.AllowFramebufferFetch=1 reverts the mobile deferred renderer from multi-pass to single-pass/Vulkan subpasses, which was the default before UE 5.7.
-; Multi-pass deferred is disabled for VR Template for 5.7.1 due to rendering issues with mobile multi-view.
-r.Mobile.AllowFramebufferFetch=1
\ No newline at end of file
+; This reverts the mobile deferred renderer to single-pass/subpass rendering, which was the default before UE 5.7.
+; Multi-pass deferred is currently disabled for VR Template due to rendering issues with mobile multi-view.
+r.Mobile.AllowFramebufferFetch=1
+
+; PSO precaching is extremely slow on first launch for standlone XR devices.
+; Projects are recommended to use Bundled PSO Caches to avoid stuttering.
+; Alternatively, re-enable PSO precaching and conceal it with a loading screen by querying FShaderPipelineCache::NumPrecompilesRemaining().
+r.PSOPrecaching=0
\ No newline at end of file
