XR | XR Template | bPackageForOpenXRImmersive is now enabled by default, which adds the standard OpenXR intent filter: org.khronos.openxr.intent.category.IMMERSIVE_HMD. While bPackageForMetaQuest is still enabled by default, it can be disabled when packaging for non-Meta Quest Android OpenXR devices, as it primarily adds additions to the Android manifest that are specific to Meta Quest.

#Synced-CL 52514612

diff --git a/Templates/TP_VirtualRealityBP/Config/DefaultEngine.ini b/Templates/TP_VirtualRealityBP/Config/DefaultEngine.ini
--- a/Templates/TP_VirtualRealityBP/Config/DefaultEngine.ini
+++ b/Templates/TP_VirtualRealityBP/Config/DefaultEngine.ini
@@ -45,6 +45,7 @@ ExtraApplicationSettings=<meta-data android:name="com.oculus.supportedDevices" a
 bPackageForMetaQuest=True
 MinSDKVersion=32
 TargetSDKVersion=32
+bPackageForOpenXRImmersive=True
 
 [/Script/WindowsTargetPlatform.WindowsTargetSettings]
 DefaultGraphicsRHI=DefaultGraphicsRHI_Default
