Disabling bPackageForOpenXRImmersive as keeping both enabled is causing issues on certain devices. Will continue to default bPackageForMetaQuest for now, but the settings both work independently for the right device(s).

#Synced-CL 52514612

diff --git a/Templates/TP_VirtualRealityBP/Config/DefaultEngine.ini b/Templates/TP_VirtualRealityBP/Config/DefaultEngine.ini
--- a/Templates/TP_VirtualRealityBP/Config/DefaultEngine.ini
+++ b/Templates/TP_VirtualRealityBP/Config/DefaultEngine.ini
@@ -45,7 +45,6 @@ ExtraApplicationSettings=<meta-data android:name="com.oculus.supportedDevices" a
 bPackageForMetaQuest=True
 MinSDKVersion=32
 TargetSDKVersion=32
-bPackageForOpenXRImmersive=True
 
 [/Script/WindowsTargetPlatform.WindowsTargetSettings]
 DefaultGraphicsRHI=DefaultGraphicsRHI_Default
